Add ScriptableToolsFramework plugin and ScriptableToolsEditorMode.
ScriptableToolsFramework provides base class UScriptableInteractiveTool which can be extended with Blueprints to implement standard InteractiveTool functionality, including per-frame Tick/Render/DrawHUD, Property Sets, Property Watchers, User Messages, and Gizmos.
UScriptableSingleClickTool and UScriptableClickDragTool "base tools" extend UScriptableInteractiveTool with support for click and drag interactions, hover, and modifier keys.
EditorScriptableToolsFramework module adds Editor-only subclasses UEditorScriptableInteractiveTool, UEditorScriptableSingleClickTool, and UEditorScriptableClickDragTool

ScriptableToolsEditorMode adds a new Editor Mode which automatically displays all UScriptableInteractiveTool instances in the current project, in a sidebar Tool Palette, with Tool interaction similar to Modeling Mode.

diff --git a/Engine/Plugins/Experimental/ScriptableToolsEditorMode/ScriptableToolsEditorMode.uplugin b/Engine/Plugins/Experimental/ScriptableToolsEditorMode/ScriptableToolsEditorMode.uplugin
new file mode 100644
--- /dev/null
+++ b/Engine/Plugins/Experimental/ScriptableToolsEditorMode/ScriptableToolsEditorMode.uplugin
@@ -0,0 +1,33 @@
+{
+	"FileVersion": 3,
+	"Version": 1,
+	"VersionName": "0.1",
+	"FriendlyName": "Scriptable Tools Editor Mode",
+	"Description": "Editor Mode for Scriptable Tools",
+	"Category": "Editor",
+	"CreatedBy": "Epic Games, Inc.",
+	"CreatedByURL": "http://epicgames.com",
+	"DocsURL": "",
+	"MarketplaceURL": "",
+	"SupportURL": "",
+	"CanContainContent": true,
+	"IsExperimentalVersion": true,
+	"Installed": false,
+	"Modules": [
+		{
+			"Name": "ScriptableToolsEditorMode",
+			"Type": "Editor",
+			"LoadingPhase": "Default"
+		}
+	],
+	"Plugins": [
+		{
+			"Name": "ScriptableToolsFramework",
+			"Enabled": true
+		},
+		{
+			"Name": "MeshModelingToolset",
+			"Enabled": true
+		}	
+	]
+}
diff --git a/Engine/Plugins/Experimental/ScriptableToolsEditorMode/Source/ScriptableToolsEditorMode/Private/ScriptableToolsEditorMode.cpp b/Engine/Plugins/Experimental/ScriptableToolsEditorMode/Source/ScriptableToolsEditorMode/Private/ScriptableToolsEditorMode.cpp
new file mode 100644
--- /dev/null
+++ b/Engine/Plugins/Experimental/ScriptableToolsEditorMode/Source/ScriptableToolsEditorMode/Private/ScriptableToolsEditorMode.cpp
@@ -0,0 +1,438 @@
+// Copyright Epic Games, Inc. All Rights Reserved.
+
+#include "ScriptableToolsEditorMode.h"
+#include "Toolkits/ToolkitManager.h"
+#include "Modules/ModuleManager.h"
+#include "LevelEditor.h"
+#include "IAssetViewport.h"
+#include "SLevelViewport.h"
+#include "EditorViewportClient.h"
+#include "EditorModeManager.h"
+#include "UnrealWidget.h"
+#include "Framework/Commands/UICommandList.h"
+
+#include "ScriptableToolsEditorModeToolkit.h"
+#include "ScriptableToolsEditorModeManagerCommands.h"
+
+#include "BaseGizmos/TransformGizmoUtil.h"
+#include "Snapping/ModelingSceneSnappingManager.h"
+
+#include "ScriptableInteractiveTool.h"
+#include "ScriptableToolBuilder.h"
+#include "ScriptableToolSet.h"
+
+#include "InteractiveToolQueryInterfaces.h" // IInteractiveToolExclusiveToolAPI
+
+
+#define LOCTEXT_NAMESPACE "UScriptableToolsEditorMode"
+
+const FEditorModeID UScriptableToolsEditorMode::EM_ScriptableToolsEditorModeId = TEXT("EM_ScriptableToolsEditorMode");
+
+
+namespace
+{
+FString GetToolName(const UInteractiveTool& Tool)
+{
+	const FString* ToolName = FTextInspector::GetSourceString(Tool.GetToolInfo().ToolDisplayName);
+	return ToolName ? *ToolName : FString(TEXT("<Invalid ToolName>"));
+}
+}
+
+UScriptableToolsEditorMode::UScriptableToolsEditorMode()
+{
+	Info = FEditorModeInfo(
+		EM_ScriptableToolsEditorModeId,
+		LOCTEXT("ScriptableToolsEditorModeName", "Scriptable Tools"),
+		FSlateIcon("ScriptableToolsEditorModeStyle", "LevelEditor.ScriptableToolsEditorMode", "LevelEditor.ScriptableToolsEditorMode.Small"),
+		true,
+		999999);
+}
+
+UScriptableToolsEditorMode::UScriptableToolsEditorMode(FVTableHelper& Helper)
+	: UBaseLegacyWidgetEdMode(Helper)
+{
+}
+
+UScriptableToolsEditorMode::~UScriptableToolsEditorMode()
+{
+}
+
+bool UScriptableToolsEditorMode::ProcessEditDelete()
+{
+	if (UEdMode::ProcessEditDelete())
+	{
+		return true;
+	}
+
+	// for now we disable deleting in an Accept-style tool because it can result in crashes if we are deleting target object
+	if ( GetToolManager()->HasAnyActiveTool() && GetToolManager()->GetActiveTool(EToolSide::Mouse)->HasAccept() )
+	{
+		GetToolManager()->DisplayMessage(
+			LOCTEXT("CannotDeleteWarning", "Cannot delete objects while this Tool is active"), EToolMessageLevel::UserWarning);
+		return true;
+	}
+
+	return false;
+}
+
+
+bool UScriptableToolsEditorMode::ProcessEditCut()
+{
+	// for now we disable deleting in an Accept-style tool because it can result in crashes if we are deleting target object
+	if (GetToolManager()->HasAnyActiveTool() && GetToolManager()->GetActiveTool(EToolSide::Mouse)->HasAccept())
+	{
+		GetToolManager()->DisplayMessage(
+			LOCTEXT("CannotCutWarning", "Cannot cut objects while this Tool is active"), EToolMessageLevel::UserWarning);
+		return true;
+	}
+
+	return false;
+}
+
+
+void UScriptableToolsEditorMode::ActorSelectionChangeNotify()
+{
+	// would like to clear selection here, but this is called multiple times, including after a transaction when
+	// we cannot identify that the selection should not be cleared
+}
+
+
+bool UScriptableToolsEditorMode::CanAutoSave() const
+{
+	// prevent autosave if any tool is active
+	return GetToolManager()->HasAnyActiveTool() == false;
+}
+
+bool UScriptableToolsEditorMode::ShouldDrawWidget() const
+{ 
+	// hide standard xform gizmo if we have an active tool
+	if (GetInteractiveToolsContext() != nullptr && GetToolManager()->HasAnyActiveTool())
+	{
+		return false;
+	}
+
+	return UBaseLegacyWidgetEdMode::ShouldDrawWidget(); 
+}
+
+void UScriptableToolsEditorMode::Tick(FEditorViewportClient* ViewportClient, float DeltaTime)
+{
+	Super::Tick(ViewportClient, DeltaTime);
+
+	if (Toolkit.IsValid())
+	{
+		FScriptableToolsEditorModeToolkit* ModeToolkit = (FScriptableToolsEditorModeToolkit*)Toolkit.Get();
+		ModeToolkit->EnableShowRealtimeWarning(ViewportClient->IsRealtime() == false);
+	}
+}
+
+
+void UScriptableToolsEditorMode::Enter()
+{
+	UEdMode::Enter();
+
+	// listen to post-build
+	GetToolManager()->OnToolPostBuild.AddUObject(this, &UScriptableToolsEditorMode::OnToolPostBuild);
+
+	//// forward shutdown requests
+	//GetToolManager()->OnToolShutdownRequest.BindLambda([this](UInteractiveToolManager*, UInteractiveTool* Tool, EToolShutdownType ShutdownType)
+	//{
+	//	GetInteractiveToolsContext()->EndTool(ShutdownType); 
+	//	return true;
+	//});
+
+	// register gizmo helper
+	UE::TransformGizmoUtil::RegisterTransformGizmoContextObject(GetInteractiveToolsContext());
+
+	// register snapping manager
+	UE::Geometry::RegisterSceneSnappingManager(GetInteractiveToolsContext());
+	//SceneSnappingManager = UE::Geometry::FindModelingSceneSnappingManager(GetToolManager());
+
+	const FScriptableToolsEditorModeManagerCommands& ModeToolCommands = FScriptableToolsEditorModeManagerCommands::Get();
+
+	// enable realtime viewport override
+	ConfigureRealTimeViewportsOverride(true);
+
+
+	ScriptableTools = NewObject<UScriptableToolSet>(this);
+	// find all the Tool Blueprints
+	ScriptableTools->ReinitializeScriptableTools();
+	// register each of them with ToolManager
+	ScriptableTools->ForEachScriptableTool([&](UClass* ToolClass, UInteractiveToolBuilder* ToolBuilder) 
+	{
+		FString UseName = ToolClass->GetName();
+		GetToolManager(EToolsContextScope::EdMode)->RegisterToolType(UseName, ToolBuilder);
+	});
+
+	// todoz
+	GetToolManager()->SelectActiveToolType(EToolSide::Left, TEXT("BeginMeshInspectorTool"));
+
+	BlueprintPreCompileHandle = GEditor->OnBlueprintPreCompile().AddUObject(this, &UScriptableToolsEditorMode::OnBlueprintPreCompile); 
+
+	// do any toolkit UI initialization that depends on the mode setup above
+	if (Toolkit.IsValid())
+	{
+		FScriptableToolsEditorModeToolkit* ModeToolkit = (FScriptableToolsEditorModeToolkit*)Toolkit.Get();
+		ModeToolkit->InitializeAfterModeSetup();
+		ModeToolkit->ForceToolPaletteRebuild();
+	}
+}
+
+
+void UScriptableToolsEditorMode::OnBlueprintPreCompile(UBlueprint* Blueprint)
+{
+	// would be nice to only do this if the Blueprint being compiled is being used for this tool...
+	if (GetToolManager()->HasActiveTool(EToolSide::Left))
+	{
+		GetToolManager()->DeactivateTool(EToolSide::Left, EToolShutdownType::Cancel);
+	}
+
+	// Probably not necessary to always rebuild the palette here. But currently this lets us respond
+	// to changes in the tool name/setting/etc
+	if (Toolkit.IsValid())
+	{
+		FScriptableToolsEditorModeToolkit* ModeToolkit = (FScriptableToolsEditorModeToolkit*)Toolkit.Get();
+		ModeToolkit->ForceToolPaletteRebuild();
+	}
+}
+
+
+void UScriptableToolsEditorMode::Exit()
+{
+	GEditor->OnBlueprintPreCompile().Remove(BlueprintPreCompileHandle);
+
+	// exit any exclusive active tools w/ cancel
+	if (UInteractiveTool* ActiveTool = GetToolManager()->GetActiveTool(EToolSide::Left))
+	{
+		if (Cast<IInteractiveToolExclusiveToolAPI>(ActiveTool))
+		{
+			GetToolManager()->DeactivateTool(EToolSide::Left, EToolShutdownType::Cancel);
+		}
+	}
+
+	UE::Geometry::DeregisterSceneSnappingManager(GetInteractiveToolsContext());
+	UE::TransformGizmoUtil::DeregisterTransformGizmoContextObject(GetInteractiveToolsContext());
+
+
+	// deregister transform gizmo context object
+	UE::TransformGizmoUtil::DeregisterTransformGizmoContextObject(GetInteractiveToolsContext());
+	
+	// clear realtime viewport override
+	ConfigureRealTimeViewportsOverride(false);
+
+	// Call base Exit method to ensure proper cleanup
+	UEdMode::Exit();
+}
+
+void UScriptableToolsEditorMode::OnToolsContextRender(IToolsContextRenderAPI* RenderAPI)
+{
+}
+
+bool UScriptableToolsEditorMode::ShouldToolStartBeAllowed(const FString& ToolIdentifier) const
+{
+	if (UInteractiveToolManager* Manager = GetToolManager())
+	{
+		if (UInteractiveTool* Tool = Manager->GetActiveTool(EToolSide::Left))
+		{
+			IInteractiveToolExclusiveToolAPI* ExclusiveAPI = Cast<IInteractiveToolExclusiveToolAPI>(Tool);
+			if (ExclusiveAPI)
+			{
+				return false;
+			}
+		}
+	}
+	return Super::ShouldToolStartBeAllowed(ToolIdentifier);
+}
+
+
+
+void UScriptableToolsEditorMode::CreateToolkit()
+{
+	Toolkit = MakeShareable(new FScriptableToolsEditorModeToolkit);
+}
+
+void UScriptableToolsEditorMode::OnToolPostBuild(
+	UInteractiveToolManager* InToolManager, EToolSide InSide, 
+	UInteractiveTool* InBuiltTool, UInteractiveToolBuilder* InToolBuilder, const FToolBuilderState& ToolState)
+{
+}
+
+void UScriptableToolsEditorMode::OnToolStarted(UInteractiveToolManager* Manager, UInteractiveTool* Tool)
+{
+}
+
+void UScriptableToolsEditorMode::OnToolEnded(UInteractiveToolManager* Manager, UInteractiveTool* Tool)
+{
+}
+
+void UScriptableToolsEditorMode::BindCommands()
+{
+	const FScriptableToolsEditorModeManagerCommands& ToolManagerCommands = FScriptableToolsEditorModeManagerCommands::Get();
+	const TSharedRef<FUICommandList>& CommandList = Toolkit->GetToolkitCommands();
+
+	CommandList->MapAction(
+		ToolManagerCommands.AcceptActiveTool,
+		FExecuteAction::CreateLambda([this]() { 
+			GetInteractiveToolsContext()->EndTool(EToolShutdownType::Accept); 
+		}),
+		FCanExecuteAction::CreateLambda([this]() { return GetInteractiveToolsContext()->CanAcceptActiveTool(); }),
+		FGetActionCheckState(),
+		FIsActionButtonVisible::CreateLambda([this]() {return GetInteractiveToolsContext()->ActiveToolHasAccept(); }),
+		EUIActionRepeatMode::RepeatDisabled
+		);
+
+	CommandList->MapAction(
+		ToolManagerCommands.CancelActiveTool,
+		FExecuteAction::CreateLambda([this]() { GetInteractiveToolsContext()->EndTool(EToolShutdownType::Cancel); }),
+		FCanExecuteAction::CreateLambda([this]() { return GetInteractiveToolsContext()->CanCancelActiveTool(); }),
+		FGetActionCheckState(),
+		FIsActionButtonVisible::CreateLambda([this]() {return GetInteractiveToolsContext()->ActiveToolHasAccept(); }),
+		EUIActionRepeatMode::RepeatDisabled
+		);
+
+	CommandList->MapAction(
+		ToolManagerCommands.CompleteActiveTool,
+		FExecuteAction::CreateLambda([this]() { GetInteractiveToolsContext()->EndTool(EToolShutdownType::Completed); }),
+		FCanExecuteAction::CreateLambda([this]() { return GetInteractiveToolsContext()->CanCompleteActiveTool(); }),
+		FGetActionCheckState(),
+		FIsActionButtonVisible::CreateLambda([this]() {return GetInteractiveToolsContext()->CanCompleteActiveTool(); }),
+		EUIActionRepeatMode::RepeatDisabled
+		);
+
+	// These aren't activated by buttons but have default chords that bind the keypresses to the action.
+	CommandList->MapAction(
+		ToolManagerCommands.AcceptOrCompleteActiveTool,
+		FExecuteAction::CreateLambda([this]() { AcceptActiveToolActionOrTool(); }),
+		FCanExecuteAction::CreateLambda([this]() {
+				return GetInteractiveToolsContext()->CanAcceptActiveTool() || GetInteractiveToolsContext()->CanCompleteActiveTool();
+			}),
+		FGetActionCheckState(),
+		FIsActionButtonVisible(),
+		EUIActionRepeatMode::RepeatDisabled);
+
+	CommandList->MapAction(
+		ToolManagerCommands.CancelOrCompleteActiveTool,
+		FExecuteAction::CreateLambda([this]() { CancelActiveToolActionOrTool(); }),
+		FCanExecuteAction::CreateLambda([this]() {
+				return GetInteractiveToolsContext()->CanCompleteActiveTool() || GetInteractiveToolsContext()->CanCancelActiveTool();
+			}),
+		FGetActionCheckState(),
+		FIsActionButtonVisible(),
+		EUIActionRepeatMode::RepeatDisabled);
+}
+
+
+void UScriptableToolsEditorMode::AcceptActiveToolActionOrTool()
+{
+	// if we have an active Tool that implements 
+	if (GetToolManager()->HasAnyActiveTool())
+	{
+		UInteractiveTool* Tool = GetToolManager()->GetActiveTool(EToolSide::Mouse);
+		IInteractiveToolNestedAcceptCancelAPI* CancelAPI = Cast<IInteractiveToolNestedAcceptCancelAPI>(Tool);
+		if (CancelAPI && CancelAPI->SupportsNestedAcceptCommand() && CancelAPI->CanCurrentlyNestedAccept())
+		{
+			bool bAccepted = CancelAPI->ExecuteNestedAcceptCommand();
+			if (bAccepted)
+			{
+				return;
+			}
+		}
+	}
+
+	const EToolShutdownType ShutdownType = GetInteractiveToolsContext()->CanAcceptActiveTool() ? EToolShutdownType::Accept : EToolShutdownType::Completed;
+	GetInteractiveToolsContext()->EndTool(ShutdownType);
+}
+
+
+void UScriptableToolsEditorMode::CancelActiveToolActionOrTool()
+{
+	// if we have an active Tool that implements 
+	if (GetToolManager()->HasAnyActiveTool())
+	{
+		UInteractiveTool* Tool = GetToolManager()->GetActiveTool(EToolSide::Mouse);
+		IInteractiveToolNestedAcceptCancelAPI* CancelAPI = Cast<IInteractiveToolNestedAcceptCancelAPI>(Tool);
+		if (CancelAPI && CancelAPI->SupportsNestedCancelCommand() && CancelAPI->CanCurrentlyNestedCancel())
+		{
+			bool bCancelled = CancelAPI->ExecuteNestedCancelCommand();
+			if (bCancelled)
+			{
+				return;
+			}
+		}
+	}
+
+	const EToolShutdownType ShutdownType = GetInteractiveToolsContext()->CanCancelActiveTool() ? EToolShutdownType::Cancel : EToolShutdownType::Completed;
+	GetInteractiveToolsContext()->EndTool(ShutdownType);
+}
+
+
+bool UScriptableToolsEditorMode::ComputeBoundingBoxForViewportFocus(AActor* Actor, UPrimitiveComponent* PrimitiveComponent, FBox& InOutBox) const
+{
+	auto ProcessFocusBoxFunc = [](FBox& FocusBoxInOut)
+	{
+		double MaxDimension = FocusBoxInOut.GetExtent().GetMax();
+		double ExpandAmount = (MaxDimension > SMALL_NUMBER) ? (MaxDimension * 0.2) : 25;		// 25 is a bit arbitrary here...
+		FocusBoxInOut = FocusBoxInOut.ExpandBy(MaxDimension * 0.2);
+	};
+
+	// if Tool supports custom Focus box, use that
+	if (GetToolManager()->HasAnyActiveTool())
+	{
+		UInteractiveTool* Tool = GetToolManager()->GetActiveTool(EToolSide::Mouse);
+		IInteractiveToolCameraFocusAPI* FocusAPI = Cast<IInteractiveToolCameraFocusAPI>(Tool);
+		if (FocusAPI && FocusAPI->SupportsWorldSpaceFocusBox() )
+		{
+			InOutBox = FocusAPI->GetWorldSpaceFocusBox();
+			if (InOutBox.IsValid)
+			{
+				ProcessFocusBoxFunc(InOutBox);
+				return true;
+			}
+		}
+	}
+
+	// fallback to base focus behavior
+	return false;
+}
+
+
+bool UScriptableToolsEditorMode::GetPivotForOrbit(FVector& OutPivot) const
+{
+	if (GCurrentLevelEditingViewportClient)
+	{
+		OutPivot = GCurrentLevelEditingViewportClient->GetViewTransform().GetLookAt();
+		return true;
+	}
+	return false;
+}
+
+
+
+void UScriptableToolsEditorMode::ConfigureRealTimeViewportsOverride(bool bEnable)
+{
+	FLevelEditorModule& LevelEditorModule = FModuleManager::GetModuleChecked<FLevelEditorModule>("LevelEditor");
+	TSharedPtr<ILevelEditor> LevelEditor = LevelEditorModule.GetFirstLevelEditor();
+	if (LevelEditor.IsValid())
+	{
+		TArray<TSharedPtr<SLevelViewport>> Viewports = LevelEditor->GetViewports();
+		for (const TSharedPtr<SLevelViewport>& ViewportWindow : Viewports)
+		{
+			if (ViewportWindow.IsValid())
+			{
+				FEditorViewportClient& Viewport = ViewportWindow->GetAssetViewportClient();
+				const FText SystemDisplayName = LOCTEXT("RealtimeOverrideMessage_ScriptableToolsMode", "ScriptableTools Mode");
+				if (bEnable)
+				{
+					Viewport.AddRealtimeOverride(bEnable, SystemDisplayName);
+				}
+				else
+				{
+					Viewport.RemoveRealtimeOverride(SystemDisplayName, false);
+				}
+			}
+		}
+	}
+}
+
+
+
+#undef LOCTEXT_NAMESPACE
diff --git a/Engine/Plugins/Experimental/ScriptableToolsEditorMode/Source/ScriptableToolsEditorMode/Private/ScriptableToolsEditorModeManagerCommands.cpp b/Engine/Plugins/Experimental/ScriptableToolsEditorMode/Source/ScriptableToolsEditorMode/Private/ScriptableToolsEditorModeManagerCommands.cpp
new file mode 100644
--- /dev/null
+++ b/Engine/Plugins/Experimental/ScriptableToolsEditorMode/Source/ScriptableToolsEditorMode/Private/ScriptableToolsEditorModeManagerCommands.cpp
@@ -0,0 +1,66 @@
+// Copyright Epic Games, Inc. All Rights Reserved.
+
+#include "ScriptableToolsEditorModeManagerCommands.h"
+#include "InputCoreTypes.h"
+#include "ScriptableToolsEditorModeStyle.h"
+
+#define LOCTEXT_NAMESPACE "ScriptableToolsEditorModeManagerCommands"
+
+
+
+FScriptableToolsEditorModeManagerCommands::FScriptableToolsEditorModeManagerCommands() :
+	TCommands<FScriptableToolsEditorModeManagerCommands>(
+		"ScriptableToolsEditorModeToolCommands", // Context name for fast lookup
+		NSLOCTEXT("Contexts", "ScriptableToolsEditorModeToolCommands", "ScriptableTools Mode - Tools"), // Localized context name for displaying
+		NAME_None, // Parent
+		FScriptableToolsEditorModeStyle::Get()->GetStyleSetName() // Icon Style Set
+		)
+{
+}
+
+
+
+TSharedPtr<FUICommandInfo> FScriptableToolsEditorModeManagerCommands::FindToolByName(FString Name, bool& bFound) const
+{
+	bFound = false;
+	for (const FStartToolCommand& Command : RegisteredTools)
+	{
+		if (Command.ToolUIName.Equals(Name, ESearchCase::IgnoreCase)
+		 || (Command.ToolCommand.IsValid() && Command.ToolCommand->GetLabel().ToString().Equals(Name, ESearchCase::IgnoreCase)))
+		{
+			bFound = true;
+			return Command.ToolCommand;
+		}
+	}
+	return TSharedPtr<FUICommandInfo>();
+}
+
+
+void FScriptableToolsEditorModeManagerCommands::RegisterCommands()
+{
+	// this has to be done with a compile-time macro because UI_COMMAND expands to LOCTEXT macros
+#define REGISTER_TOOL_COMMAND(ToolCommandInfo, ToolName, ToolTip) \
+		UI_COMMAND(ToolCommandInfo, ToolName, ToolTip, EUserInterfaceActionType::ToggleButton, FInputChord()); \
+		RegisteredTools.Add(FStartToolCommand{ ToolName, ToolCommandInfo });
+
+	REGISTER_TOOL_COMMAND(BeginMeshInspectorTool, "Inspct", "Inspect mesh attributes");
+	REGISTER_TOOL_COMMAND(BeginPhysicsInspectorTool, "PInspct", "Inspect physics geometry for selected meshes");
+	REGISTER_TOOL_COMMAND(BeginSetCollisionGeometryTool, "Msh2Coll", "Convert selected meshes to Simple Collision Geometry (for last selected)");
+
+
+
+	UI_COMMAND(AcceptActiveTool, "Accept", "Accept the active tool", EUserInterfaceActionType::Button, FInputChord());
+	UI_COMMAND(CancelActiveTool, "Cancel", "Cancel the active tool", EUserInterfaceActionType::Button, FInputChord());
+	UI_COMMAND(CompleteActiveTool, "Done", "Complete the active tool", EUserInterfaceActionType::Button, FInputChord());
+
+	// Note that passing a chord into one of these calls hooks the key press to the respective action. 
+	UI_COMMAND(AcceptOrCompleteActiveTool, "Accept or Complete", "Accept or complete the active tool", EUserInterfaceActionType::Button, FInputChord(EKeys::Enter));
+	UI_COMMAND(CancelOrCompleteActiveTool, "Cancel or Complete", "Cancel or complete the active tool", EUserInterfaceActionType::Button, FInputChord(EKeys::Escape));
+	
+#undef REGISTER_TOOL_COMMAND
+}
+
+
+
+
+#undef LOCTEXT_NAMESPACE
diff --git a/Engine/Plugins/Experimental/ScriptableToolsEditorMode/Source/ScriptableToolsEditorMode/Private/ScriptableToolsEditorModeModule.cpp b/Engine/Plugins/Experimental/ScriptableToolsEditorMode/Source/ScriptableToolsEditorMode/Private/ScriptableToolsEditorModeModule.cpp
new file mode 100644
--- /dev/null
+++ b/Engine/Plugins/Experimental/ScriptableToolsEditorMode/Source/ScriptableToolsEditorMode/Private/ScriptableToolsEditorModeModule.cpp
@@ -0,0 +1,36 @@
+// Copyright Epic Games, Inc. All Rights Reserved.
+
+#include "ScriptableToolsEditorModeModule.h"
+#include "ScriptableToolsEditorModeManagerCommands.h"
+#include "ScriptableToolsEditorModeStyle.h"
+#include "Misc/CoreDelegates.h"
+
+#define LOCTEXT_NAMESPACE "FScriptableToolsEditorModeModule"
+
+void FScriptableToolsEditorModeModule::StartupModule()
+{
+	// This code will execute after your module is loaded into memory; the exact timing is specified in the .uplugin file per-module
+
+	FCoreDelegates::OnPostEngineInit.AddRaw(this, &FScriptableToolsEditorModeModule::OnPostEngineInit);
+}
+
+void FScriptableToolsEditorModeModule::ShutdownModule()
+{
+	FCoreDelegates::OnPostEngineInit.RemoveAll(this);
+
+	// This function may be called during shutdown to clean up your module.  For modules that support dynamic reloading,
+	// we call this function before unloading the module.
+
+	FScriptableToolsEditorModeManagerCommands::Unregister();
+	FScriptableToolsEditorModeStyle::Shutdown();
+}
+
+void FScriptableToolsEditorModeModule::OnPostEngineInit()
+{
+	FScriptableToolsEditorModeStyle::Initialize();
+	FScriptableToolsEditorModeManagerCommands::Register();
+}
+
+#undef LOCTEXT_NAMESPACE
+	
+IMPLEMENT_MODULE(FScriptableToolsEditorModeModule, ScriptableToolsEditorMode)
\ No newline at end of file
diff --git a/Engine/Plugins/Experimental/ScriptableToolsEditorMode/Source/ScriptableToolsEditorMode/Private/ScriptableToolsEditorModeStyle.cpp b/Engine/Plugins/Experimental/ScriptableToolsEditorMode/Source/ScriptableToolsEditorMode/Private/ScriptableToolsEditorModeStyle.cpp
new file mode 100644
--- /dev/null
+++ b/Engine/Plugins/Experimental/ScriptableToolsEditorMode/Source/ScriptableToolsEditorMode/Private/ScriptableToolsEditorModeStyle.cpp
@@ -0,0 +1,89 @@
+// Copyright Epic Games, Inc. All Rights Reserved.
+
+#include "ScriptableToolsEditorModeStyle.h"
+#include "Styling/SlateStyleRegistry.h"
+#include "Styling/SlateTypes.h"
+#include "Styling/CoreStyle.h"
+#include "Styling/AppStyle.h"
+#include "Interfaces/IPluginManager.h"
+#include "SlateOptMacros.h"
+#include "Styling/SlateStyleMacros.h"
+
+
+#define IMAGE_PLUGIN_BRUSH( RelativePath, ... ) FSlateImageBrush( FScriptableToolsEditorModeStyle::InContent( RelativePath, ".png" ), __VA_ARGS__ )
+
+// This is to fix the issue that SlateStyleMacros like IMAGE_BRUSH look for RootToContentDir but StyleSet->RootToContentDir is how this style is set up
+#define RootToContentDir StyleSet->RootToContentDir
+
+FString FScriptableToolsEditorModeStyle::InContent(const FString& RelativePath, const ANSICHAR* Extension)
+{
+	static FString ContentDir = IPluginManager::Get().FindPlugin(TEXT("ScriptableToolsEditorMode"))->GetContentDir();
+	return (ContentDir / RelativePath) + Extension;
+}
+
+TSharedPtr< FSlateStyleSet > FScriptableToolsEditorModeStyle::StyleSet = nullptr;
+TSharedPtr< class ISlateStyle > FScriptableToolsEditorModeStyle::Get() { return StyleSet; }
+
+FName FScriptableToolsEditorModeStyle::GetStyleSetName()
+{
+	static FName ScriptableToolsEditorModeStyleName(TEXT("ScriptableToolsEditorModeStyle"));
+	return ScriptableToolsEditorModeStyleName;
+}
+
+const FSlateBrush* FScriptableToolsEditorModeStyle::GetBrush(FName PropertyName, const ANSICHAR* Specifier)
+{
+	return Get()->GetBrush(PropertyName, Specifier);
+}
+
+
+BEGIN_SLATE_FUNCTION_BUILD_OPTIMIZATION
+
+void FScriptableToolsEditorModeStyle::Initialize()
+{
+	// Const icon sizes
+	const FVector2D Icon8x8(8.0f, 8.0f);
+	const FVector2D Icon16x16(16.0f, 16.0f);
+	const FVector2D Icon20x20(20.0f, 20.0f);
+	const FVector2D Icon28x28(28.0f, 28.0f);
+	const FVector2D Icon40x40(40.0f, 40.0f);
+	const FVector2D Icon120(120.0f, 120.0f);
+
+	// Only register once
+	if (StyleSet.IsValid())
+	{
+		return;
+	}
+
+	StyleSet = MakeShareable(new FSlateStyleSet(GetStyleSetName()));
+	StyleSet->SetContentRoot(FPaths::EnginePluginsDir() / TEXT("Editor/ScriptableToolsEditorMode/Content"));
+	StyleSet->SetCoreContentRoot(FPaths::EngineContentDir() / TEXT("Slate"));
+
+	const FTextBlockStyle& NormalText = FAppStyle::Get().GetWidgetStyle<FTextBlockStyle>("NormalText");
+
+	{
+		StyleSet->Set("LevelEditor.ScriptableToolsEditorMode", new IMAGE_PLUGIN_BRUSH("Icons/ScriptableToolsEditorMode_Icon_40x", FVector2D(20.0f, 20.0f)));
+
+		StyleSet->Set("ScriptableToolsEditorModeToolCommands.DefaultToolIcon", new IMAGE_PLUGIN_BRUSH("Icons/Tool_DefaultIcon_40px", Icon20x20));
+		StyleSet->Set("ScriptableToolsEditorModeToolCommands.DefaultToolIcon.Small", new IMAGE_PLUGIN_BRUSH("Icons/Tool_DefaultIcon_40px", Icon20x20));
+
+	}
+
+	FSlateStyleRegistry::RegisterSlateStyle(*StyleSet.Get());
+};
+
+END_SLATE_FUNCTION_BUILD_OPTIMIZATION
+
+#undef IMAGE_PLUGIN_BRUSH
+#undef IMAGE_BRUSH
+#undef BOX_BRUSH
+#undef DEFAULT_FONT
+
+void FScriptableToolsEditorModeStyle::Shutdown()
+{
+	if (StyleSet.IsValid())
+	{
+		FSlateStyleRegistry::UnRegisterSlateStyle(*StyleSet.Get());
+		ensure(StyleSet.IsUnique());
+		StyleSet.Reset();
+	}
+}
diff --git a/Engine/Plugins/Experimental/ScriptableToolsEditorMode/Source/ScriptableToolsEditorMode/Private/ScriptableToolsEditorModeToolkit.cpp b/Engine/Plugins/Experimental/ScriptableToolsEditorMode/Source/ScriptableToolsEditorMode/Private/ScriptableToolsEditorModeToolkit.cpp
new file mode 100644
--- /dev/null
+++ b/Engine/Plugins/Experimental/ScriptableToolsEditorMode/Source/ScriptableToolsEditorMode/Private/ScriptableToolsEditorModeToolkit.cpp
@@ -0,0 +1,519 @@
+// Copyright Epic Games, Inc. All Rights Reserved.
+
+#include "ScriptableToolsEditorModeToolkit.h"
+#include "ScriptableToolsEditorMode.h"
+#include "ScriptableToolsEditorModeManagerCommands.h"
+#include "ScriptableToolsEditorModeStyle.h"
+#include "Engine/Selection.h"
+
+#include "Framework/MultiBox/MultiBoxBuilder.h"
+#include "Modules/ModuleManager.h"
+#include "IDetailsView.h"
+#include "ISettingsModule.h"
+#include "EditorModeManager.h"
+#include "Toolkits/AssetEditorModeUILayer.h"
+
+#include "InteractiveToolManager.h"
+
+#include "Widgets/Input/SButton.h"
+#include "SSimpleButton.h"
+#include "Widgets/Text/STextBlock.h"
+#include "Widgets/Layout/SSeparator.h"
+#include "Widgets/Images/SImage.h"
+#include "Widgets/Layout/SExpandableArea.h"
+#include "SPrimaryButton.h"
+#include "Framework/MultiBox/SToolBarButtonBlock.h"
+#include "Internationalization/Text.h"
+
+#include "ScriptableInteractiveTool.h"
+#include "ScriptableToolSet.h"
+
+#define LOCTEXT_NAMESPACE "FScriptableToolsEditorModeToolkit"
+
+
+
+FScriptableToolsEditorModeToolkit::FScriptableToolsEditorModeToolkit()
+{
+}
+
+
+FScriptableToolsEditorModeToolkit::~FScriptableToolsEditorModeToolkit()
+{
+	GetScriptableEditorMode()->GetInteractiveToolsContext(EToolsContextScope::EdMode)->OnToolNotificationMessage.RemoveAll(this);
+	GetScriptableEditorMode()->GetInteractiveToolsContext(EToolsContextScope::EdMode)->OnToolWarningMessage.RemoveAll(this);
+}
+
+
+void FScriptableToolsEditorModeToolkit::CustomizeModeDetailsViewArgs(FDetailsViewArgs& ArgsInOut)
+{
+	//ArgsInOut.ColumnWidth = 0.3f;
+}
+
+void FScriptableToolsEditorModeToolkit::Init(const TSharedPtr<IToolkitHost>& InitToolkitHost, TWeakObjectPtr<UEdMode> InOwningMode)
+{
+	// Have to create the ToolkitWidget here because FModeToolkit::Init() is going to ask for it and add
+	// it to the Mode panel, and not ask again afterwards. However we have to call Init() to get the 
+	// ModeDetailsView created, that we need to add to the ToolkitWidget. So, we will create the Widget
+	// here but only add the rows to it after we call Init()
+	TSharedPtr<SVerticalBox> ToolkitWidgetVBox = SNew(SVerticalBox);
+	SAssignNew(ToolkitWidget, SBorder)
+		.HAlign(HAlign_Fill)
+		.Padding(4)
+		[
+			ToolkitWidgetVBox->AsShared()
+		];
+
+	FModeToolkit::Init(InitToolkitHost, InOwningMode);
+
+	GetToolkitHost()->OnActiveViewportChanged().AddSP(this, &FScriptableToolsEditorModeToolkit::OnActiveViewportChanged);
+
+	ModeWarningArea = SNew(STextBlock)
+		.AutoWrapText(true)
+		.Font(FCoreStyle::GetDefaultFontStyle("Bold", 9))
+		.ColorAndOpacity(FSlateColor(FLinearColor(0.9f, 0.15f, 0.15f)));
+	ModeWarningArea->SetText(FText::GetEmpty());
+	ModeWarningArea->SetVisibility(EVisibility::Collapsed);
+
+	ModeHeaderArea = SNew(STextBlock)
+		.AutoWrapText(true)
+		.Font(FCoreStyle::GetDefaultFontStyle("Bold", 12));
+	ModeHeaderArea->SetText(LOCTEXT("SelectToolLabel", "Select a Tool from the Toolbar"));
+	ModeHeaderArea->SetJustification(ETextJustify::Center);
+
+
+	ToolWarningArea = SNew(STextBlock)
+		.AutoWrapText(true)
+		.Font(FCoreStyle::GetDefaultFontStyle("Bold", 9))
+		.ColorAndOpacity(FSlateColor(FLinearColor(0.9f, 0.15f, 0.15f)));
+	ToolWarningArea->SetText(FText::GetEmpty());
+
+	// add the various sections to the mode toolbox
+	ToolkitWidgetVBox->AddSlot().AutoHeight().HAlign(HAlign_Fill).Padding(5)
+		[
+			ModeWarningArea->AsShared()
+		];
+	ToolkitWidgetVBox->AddSlot().AutoHeight().HAlign(HAlign_Fill).Padding(5)
+		[
+			ModeHeaderArea->AsShared()
+		];
+	ToolkitWidgetVBox->AddSlot().AutoHeight().HAlign(HAlign_Fill).Padding(5)
+		[
+			ToolWarningArea->AsShared()
+		];
+	ToolkitWidgetVBox->AddSlot().HAlign(HAlign_Fill).FillHeight(1.f)
+		[
+			ModeDetailsView->AsShared()
+		];
+
+	ClearNotification();
+	ClearWarning();
+
+	ActiveToolName = FText::GetEmpty();
+	ActiveToolMessage = FText::GetEmpty();
+
+	GetScriptableEditorMode()->GetInteractiveToolsContext(EToolsContextScope::EdMode)->OnToolNotificationMessage.AddSP(this, &FScriptableToolsEditorModeToolkit::PostNotification);
+	GetScriptableEditorMode()->GetInteractiveToolsContext(EToolsContextScope::EdMode)->OnToolWarningMessage.AddSP(this, &FScriptableToolsEditorModeToolkit::PostWarning);
+
+	SAssignNew(ViewportOverlayWidget, SHorizontalBox)
+
+	+SHorizontalBox::Slot()
+	.HAlign(HAlign_Center)
+	.VAlign(VAlign_Bottom)
+	.Padding(FMargin(0.0f, 0.0f, 0.f, 15.f))
+	[
+		SNew(SBorder)
+		.BorderImage(FAppStyle::Get().GetBrush("EditorViewport.OverlayBrush"))
+		.Padding(8.f)
+		[
+			SNew(SHorizontalBox)
+
+			+SHorizontalBox::Slot()
+			.AutoWidth()
+			.VAlign(VAlign_Center)
+			.Padding(FMargin(0.f, 0.f, 8.f, 0.f))
+			[
+				SNew(SImage)
+				.Image_Lambda([this] () { return ActiveToolIcon; })
+			]
+
+			+SHorizontalBox::Slot()
+			.AutoWidth()
+			.VAlign(VAlign_Center)
+			.Padding(FMargin(0.f, 0.f, 8.f, 0.f))
+			[
+				SNew(STextBlock)
+				.Text(this, &FScriptableToolsEditorModeToolkit::GetActiveToolDisplayName)
+			]
+
+			+SHorizontalBox::Slot()
+			.AutoWidth()
+			.Padding(FMargin(0.0, 0.f, 2.f, 0.f))
+			[
+				SNew(SPrimaryButton)
+				.Text(LOCTEXT("OverlayAccept", "Accept"))
+				.ToolTipText(LOCTEXT("OverlayAcceptTooltip", "Accept/Commit the results of the active Tool [Enter]"))
+				.OnClicked_Lambda([this]() { GetScriptableEditorMode()->GetInteractiveToolsContext(EToolsContextScope::EdMode)->EndTool(EToolShutdownType::Accept); return FReply::Handled(); })
+				.IsEnabled_Lambda([this]() { return GetScriptableEditorMode()->GetInteractiveToolsContext(EToolsContextScope::EdMode)->CanAcceptActiveTool(); })
+				.Visibility_Lambda([this]() { return GetScriptableEditorMode()->GetInteractiveToolsContext(EToolsContextScope::EdMode)->ActiveToolHasAccept() ? EVisibility::Visible : EVisibility::Collapsed; })
+			]
+
+			+SHorizontalBox::Slot()
+			.AutoWidth()
+			.Padding(FMargin(2.0, 0.f, 0.f, 0.f))
+			[
+				SNew(SButton)
+				.Text(LOCTEXT("OverlayCancel", "Cancel"))
+				.ToolTipText(LOCTEXT("OverlayCancelTooltip", "Cancel the active Tool [Esc]"))
+				.HAlign(HAlign_Center)
+				.OnClicked_Lambda([this]() { GetScriptableEditorMode()->GetInteractiveToolsContext(EToolsContextScope::EdMode)->EndTool(EToolShutdownType::Cancel); return FReply::Handled(); })
+				.IsEnabled_Lambda([this]() { return GetScriptableEditorMode()->GetInteractiveToolsContext(EToolsContextScope::EdMode)->CanCancelActiveTool(); })
+				.Visibility_Lambda([this]() { return GetScriptableEditorMode()->GetInteractiveToolsContext(EToolsContextScope::EdMode)->ActiveToolHasAccept() ? EVisibility::Visible : EVisibility::Collapsed; })
+			]
+
+			+SHorizontalBox::Slot()
+			.AutoWidth()
+			.Padding(FMargin(2.0, 0.f, 0.f, 0.f))
+			[
+				SNew(SPrimaryButton)
+				.Text(LOCTEXT("OverlayComplete", "Complete"))
+				.ToolTipText(LOCTEXT("OverlayCompleteTooltip", "Exit the active Tool [Enter]"))
+				.OnClicked_Lambda([this]() { GetScriptableEditorMode()->GetInteractiveToolsContext(EToolsContextScope::EdMode)->EndTool(EToolShutdownType::Completed); return FReply::Handled(); })
+				.IsEnabled_Lambda([this]() { return GetScriptableEditorMode()->GetInteractiveToolsContext(EToolsContextScope::EdMode)->CanCompleteActiveTool(); })
+				.Visibility_Lambda([this]() { return GetScriptableEditorMode()->GetInteractiveToolsContext(EToolsContextScope::EdMode)->ActiveToolHasAccept() ? EVisibility::Collapsed : EVisibility::Visible; })
+			]
+		]	
+	];
+
+}
+
+
+
+
+
+void FScriptableToolsEditorModeToolkit::InitializeAfterModeSetup()
+{
+	if (bFirstInitializeAfterModeSetup)
+	{
+		bFirstInitializeAfterModeSetup = false;
+	}
+
+	UpdateActiveToolCategories();
+}
+
+
+
+void FScriptableToolsEditorModeToolkit::UpdateActiveToolCategories()
+{
+	// build tool category list
+	ActiveToolCategories.Reset();
+	UScriptableToolsEditorMode* EditorMode = Cast<UScriptableToolsEditorMode>(GetScriptableEditorMode());
+	UScriptableToolSet* ScriptableTools = EditorMode->GetActiveScriptableTools();
+	ScriptableTools->ForEachScriptableTool( [&](UClass* ToolClass, UBaseScriptableToolBuilder* ToolBuilder) 
+	{
+		UScriptableInteractiveTool* ToolCDO = Cast<UScriptableInteractiveTool>(ToolClass->GetDefaultObject());
+		if (ToolCDO->ToolCategory.IsEmpty() == false)
+		{
+			FName CategoryName(ToolCDO->ToolCategory.ToString());
+			if (ActiveToolCategories.Contains(CategoryName) == false)
+			{
+				ActiveToolCategories.Add(CategoryName, ToolCDO->ToolCategory);
+			}
+		}
+	});
+}
+
+
+
+void FScriptableToolsEditorModeToolkit::UpdateActiveToolProperties()
+{
+	UInteractiveTool* CurTool = GetScriptableEditorMode()->GetToolManager(EToolsContextScope::EdMode)->GetActiveTool(EToolSide::Left);
+	if (CurTool != nullptr)
+	{
+		ModeDetailsView->SetObjects(CurTool->GetToolProperties(true));
+	}
+}
+
+void FScriptableToolsEditorModeToolkit::InvalidateCachedDetailPanelState(UObject* ChangedObject)
+{
+	ModeDetailsView->InvalidateCachedState();
+}
+
+
+void FScriptableToolsEditorModeToolkit::PostNotification(const FText& Message)
+{
+	ClearNotification();
+
+	ActiveToolMessage = Message;
+
+	if (ModeUILayer.IsValid())
+	{
+		TSharedPtr<FAssetEditorModeUILayer> ModeUILayerPtr = ModeUILayer.Pin();
+		ActiveToolMessageHandle = GEditor->GetEditorSubsystem<UStatusBarSubsystem>()->PushStatusBarMessage(ModeUILayerPtr->GetStatusBarName(), ActiveToolMessage);
+	}
+}
+
+void FScriptableToolsEditorModeToolkit::ClearNotification()
+{
+	ActiveToolMessage = FText::GetEmpty();
+
+	if (ModeUILayer.IsValid())
+	{
+		TSharedPtr<FAssetEditorModeUILayer> ModeUILayerPtr = ModeUILayer.Pin();
+		GEditor->GetEditorSubsystem<UStatusBarSubsystem>()->PopStatusBarMessage(ModeUILayerPtr->GetStatusBarName(), ActiveToolMessageHandle);
+	}
+	ActiveToolMessageHandle.Reset();
+}
+
+
+void FScriptableToolsEditorModeToolkit::PostWarning(const FText& Message)
+{
+	ToolWarningArea->SetText(Message);
+	ToolWarningArea->SetVisibility(EVisibility::Visible);
+}
+
+void FScriptableToolsEditorModeToolkit::ClearWarning()
+{
+	ToolWarningArea->SetText(FText());
+	ToolWarningArea->SetVisibility(EVisibility::Collapsed);
+}
+
+
+
+FName FScriptableToolsEditorModeToolkit::GetToolkitFName() const
+{
+	return FName("ScriptableToolsEditorMode");
+}
+
+FText FScriptableToolsEditorModeToolkit::GetBaseToolkitName() const
+{
+	return NSLOCTEXT("ScriptableToolsEditorModeToolkit", "DisplayName", "ScriptableToolsEditorMode Tool");
+}
+
+static const FName CustomToolsTabName(TEXT("CustomTools"));
+
+void FScriptableToolsEditorModeToolkit::GetToolPaletteNames(TArray<FName>& PaletteNames) const
+{
+	PaletteNames.Reset();
+	bool bFoundUncategorized = false;
+
+	UScriptableToolsEditorMode* EditorMode = Cast<UScriptableToolsEditorMode>(GetScriptableEditorMode());
+	UInteractiveToolManager* ToolManager = EditorMode->GetToolManager();
+	UScriptableToolSet* ScriptableTools = EditorMode->GetActiveScriptableTools();
+	ScriptableTools->ForEachScriptableTool( [&](UClass* ToolClass, UBaseScriptableToolBuilder* ToolBuilder) 
+	{
+		UScriptableInteractiveTool* ToolCDO = Cast<UScriptableInteractiveTool>(ToolClass->GetDefaultObject());
+		if (ToolCDO->ToolCategory.IsEmpty() == false)
+		{
+			FName CategoryName(ToolCDO->ToolCategory.ToString());
+			if (ActiveToolCategories.Contains(CategoryName))
+			{
+				PaletteNames.AddUnique(CategoryName);
+			}
+			else
+			{
+				bFoundUncategorized = true;
+			}
+		}
+		else
+		{
+			bFoundUncategorized = true;
+		}
+	});
+
+	if ( bFoundUncategorized )
+	{
+		PaletteNames.Add(CustomToolsTabName);
+	}
+}
+
+
+FText FScriptableToolsEditorModeToolkit::GetToolPaletteDisplayName(FName Palette) const
+{ 
+	if (ActiveToolCategories.Contains(Palette))
+	{
+		return ActiveToolCategories[Palette];
+	}
+	return FText::FromName(Palette);
+}
+
+void FScriptableToolsEditorModeToolkit::BuildToolPalette(FName PaletteIndex, class FToolBarBuilder& ToolbarBuilder) 
+{
+	const FScriptableToolsEditorModeManagerCommands& Commands = FScriptableToolsEditorModeManagerCommands::Get();
+
+	static TArray<TSharedPtr<FUIAction>> ActionsHack;
+
+	bool bIsUncategorizedPalette = (PaletteIndex == CustomToolsTabName);
+
+	ActionsHack.Reset();
+	
+	// this is kind of dumb and we probably should maintain a map <FName, TArray<ToolClass>> when we build the ActiveToolCategories...
+	UScriptableToolsEditorMode* EditorMode = Cast<UScriptableToolsEditorMode>(GetScriptableEditorMode());
+	UInteractiveToolManager* ToolManager = EditorMode->GetToolManager();
+	UScriptableToolSet* ScriptableTools = EditorMode->GetActiveScriptableTools();
+	ScriptableTools->ForEachScriptableTool([&](UClass* ToolClass, UBaseScriptableToolBuilder* ToolBuilder)
+	{
+		UScriptableInteractiveTool* ToolCDO = Cast<UScriptableInteractiveTool>(ToolClass->GetDefaultObject());
+		FName UseCategoryName = ToolCDO->ToolCategory.IsEmpty() ? CustomToolsTabName : FName(ToolCDO->ToolCategory.ToString());
+		if (ActiveToolCategories.Contains(UseCategoryName) == false)
+		{
+			UseCategoryName = CustomToolsTabName;
+		}
+		if (UseCategoryName != PaletteIndex)
+		{
+			return;
+		}
+
+		FString ToolIdentifier = ToolClass->GetName();
+
+		TSharedPtr<FUIAction> NewAction = MakeShared<FUIAction>(
+		FExecuteAction::CreateLambda([this, ToolClass, ToolIdentifier, ToolManager]()
+		{
+			UScriptableInteractiveTool* ToolCDO = Cast<UScriptableInteractiveTool>(ToolClass->GetDefaultObject());
+			//UE_LOG(LogTemp, Warning, TEXT("STARTING TOOL [%s] (Class/Identifier %s)"), *ToolCDO->ToolName.ToString(), *ToolIdentifier);
+			if (ToolManager->SelectActiveToolType(EToolSide::Mouse, ToolIdentifier))
+			{
+				if (ToolManager->CanActivateTool(EToolSide::Mouse, ToolIdentifier)) 
+				{
+					bool bLaunched = ToolManager->ActivateTool(EToolSide::Mouse);
+					//ensure(bLaunched);
+				}
+			}
+			else
+			{
+				UE_LOG(LogTemp, Warning, TEXT("FAILED TO SET ACTIVE TOOL TYPE!"));
+			}
+
+		}));
+
+		ActionsHack.Add(NewAction);
+
+		FName InExtensionHook = NAME_None;
+		TAttribute<FText> Label = ToolCDO->ToolName.IsEmpty() ? LOCTEXT("EmptyToolName", "Tool") : ToolCDO->ToolName;
+		TAttribute<FText> Tooltip = ToolCDO->ToolTooltip.IsEmpty() ? FText() : ToolCDO->ToolTooltip;
+		// hardcoded icon name for now
+		TAttribute<FSlateIcon> Icon = FSlateIcon(FScriptableToolsEditorModeStyle::Get()->GetStyleSetName(), "ScriptableToolsEditorModeToolCommands.DefaultToolIcon");
+			
+		ToolbarBuilder.AddToolBarButton(*NewAction, InExtensionHook, Label, Tooltip, Icon);
+
+	});
+}
+
+void FScriptableToolsEditorModeToolkit::InvokeUI()
+{
+	FModeToolkit::InvokeUI();
+}
+
+
+
+void FScriptableToolsEditorModeToolkit::ForceToolPaletteRebuild()
+{
+	this->UpdateActiveToolCategories();
+	this->RebuildModeToolPalette();
+}
+
+
+void FScriptableToolsEditorModeToolkit::OnToolPaletteChanged(FName PaletteName) 
+{
+}
+
+
+
+void FScriptableToolsEditorModeToolkit::EnableShowRealtimeWarning(bool bEnable)
+{
+	if (bShowRealtimeWarning != bEnable)
+	{
+		bShowRealtimeWarning = bEnable;
+		UpdateShowWarnings();
+	}
+}
+
+void FScriptableToolsEditorModeToolkit::OnToolStarted(UInteractiveToolManager* Manager, UInteractiveTool* Tool)
+{
+	UpdateActiveToolProperties();
+
+	UInteractiveTool* CurTool = GetScriptableEditorMode()->GetToolManager(EToolsContextScope::EdMode)->GetActiveTool(EToolSide::Left);
+	CurTool->OnPropertySetsModified.AddSP(this, &FScriptableToolsEditorModeToolkit::UpdateActiveToolProperties);
+	CurTool->OnPropertyModifiedDirectlyByTool.AddSP(this, &FScriptableToolsEditorModeToolkit::InvalidateCachedDetailPanelState);
+
+	ModeHeaderArea->SetVisibility(EVisibility::Collapsed);
+
+	ActiveToolName = CurTool->GetToolInfo().ToolDisplayName;
+	if (UScriptableInteractiveTool* ScriptableTool = Cast<UScriptableInteractiveTool>(CurTool))
+	{
+		if ( ScriptableTool->ToolLongName.IsEmpty() == false )
+		{
+			ActiveToolName = ScriptableTool->ToolLongName;
+		}
+		else if ( ScriptableTool->ToolName.IsEmpty() == false )
+		{
+			ActiveToolName = ScriptableTool->ToolName;
+		}
+	}
+
+	// try to update icon
+	FString ActiveToolIdentifier = GetScriptableEditorMode()->GetToolManager(EToolsContextScope::EdMode)->GetActiveToolName(EToolSide::Left);
+	ActiveToolIdentifier.InsertAt(0, ".");
+	FName ActiveToolIconName = ISlateStyle::Join(FScriptableToolsEditorModeManagerCommands::Get().GetContextName(), TCHAR_TO_ANSI(*ActiveToolIdentifier));
+	ActiveToolIcon = FScriptableToolsEditorModeStyle::Get()->GetOptionalBrush(ActiveToolIconName);
+
+	GetToolkitHost()->AddViewportOverlayWidget(ViewportOverlayWidget.ToSharedRef());
+}
+
+void FScriptableToolsEditorModeToolkit::OnToolEnded(UInteractiveToolManager* Manager, UInteractiveTool* Tool)
+{
+	if (IsHosted())
+	{
+		GetToolkitHost()->RemoveViewportOverlayWidget(ViewportOverlayWidget.ToSharedRef());
+	}
+
+	ModeDetailsView->SetObject(nullptr);
+	ActiveToolName = FText::GetEmpty();
+	ModeHeaderArea->SetVisibility(EVisibility::Visible);
+	ModeHeaderArea->SetText(LOCTEXT("SelectToolLabel", "Select a Tool from the Toolbar"));
+	ClearNotification();
+	ClearWarning();
+	UInteractiveTool* CurTool = GetScriptableEditorMode()->GetToolManager(EToolsContextScope::EdMode)->GetActiveTool(EToolSide::Left);
+	if ( CurTool )
+	{
+		CurTool->OnPropertySetsModified.RemoveAll(this);
+		CurTool->OnPropertyModifiedDirectlyByTool.RemoveAll(this);
+	}
+}
+
+void FScriptableToolsEditorModeToolkit::OnActiveViewportChanged(TSharedPtr<IAssetViewport> OldViewport, TSharedPtr<IAssetViewport> NewViewport)
+{
+	// Only worry about handling this notification if we have an active tool
+	if (!ActiveToolName.IsEmpty())
+	{
+		// Check first to see if this changed because the old viewport was deleted and if not, remove our hud
+		if (OldViewport)	
+		{
+			GetToolkitHost()->RemoveViewportOverlayWidget(ViewportOverlayWidget.ToSharedRef(), OldViewport);
+		}
+
+		// Add the hud to the new viewport
+		GetToolkitHost()->AddViewportOverlayWidget(ViewportOverlayWidget.ToSharedRef(), NewViewport);
+	}
+}
+
+void FScriptableToolsEditorModeToolkit::UpdateShowWarnings()
+{
+	if (bShowRealtimeWarning )
+	{
+		if (ModeWarningArea->GetVisibility() == EVisibility::Collapsed)
+		{
+			ModeWarningArea->SetText(LOCTEXT("ScriptableToolsModeToolkitRealtimeWarning", "Realtime Mode is required for Scriptable Tools to work correctly. Please enable Realtime Mode in the Viewport Options or with the Ctrl+r hotkey."));
+			ModeWarningArea->SetVisibility(EVisibility::Visible);
+		}
+	}
+	else
+	{
+		ModeWarningArea->SetText(FText());
+		ModeWarningArea->SetVisibility(EVisibility::Collapsed);
+	}
+
+}
+
+
+
+#undef LOCTEXT_NAMESPACE
diff --git a/Engine/Plugins/Experimental/ScriptableToolsEditorMode/Source/ScriptableToolsEditorMode/Public/ScriptableToolsEditorMode.h b/Engine/Plugins/Experimental/ScriptableToolsEditorMode/Source/ScriptableToolsEditorMode/Public/ScriptableToolsEditorMode.h
new file mode 100644
--- /dev/null
+++ b/Engine/Plugins/Experimental/ScriptableToolsEditorMode/Source/ScriptableToolsEditorMode/Public/ScriptableToolsEditorMode.h
@@ -0,0 +1,93 @@
+// Copyright Epic Games, Inc. All Rights Reserved.
+
+#pragma once
+
+#include "CoreMinimal.h"
+#include "Tools/LegacyEdModeWidgetHelpers.h"
+
+#include "InputState.h"
+#include "InteractiveToolManager.h"
+#include "EdModeInteractiveToolsContext.h"
+
+#include "ScriptableToolsEditorMode.generated.h"
+
+class FUICommandList;
+class FLevelObjectsObserver;
+class UModelingSceneSnappingManager;
+class UModelingSelectionInteraction;
+class UGeometrySelectionManager;
+class UScriptableToolSet;
+class UBlueprint;
+
+UCLASS(Transient)
+class SCRIPTABLETOOLSEDITORMODE_API UScriptableToolsEditorMode : public UBaseLegacyWidgetEdMode
+{
+	GENERATED_BODY()
+public:
+	const static FEditorModeID EM_ScriptableToolsEditorModeId;
+
+	UScriptableToolsEditorMode();
+	UScriptableToolsEditorMode(FVTableHelper& Helper);
+	~UScriptableToolsEditorMode();
+	////////////////
+	// UEdMode interface
+	////////////////
+
+	virtual void Tick(FEditorViewportClient* ViewportClient, float DeltaTime) override;
+
+	virtual void ActorSelectionChangeNotify() override;
+
+	virtual bool ShouldDrawWidget() const override;
+	virtual bool ProcessEditDelete() override;
+	virtual bool ProcessEditCut() override;
+
+	virtual bool CanAutoSave() const override;
+
+	virtual bool ComputeBoundingBoxForViewportFocus(AActor* Actor, UPrimitiveComponent* PrimitiveComponent, FBox& InOutBox) const override;
+
+	virtual bool GetPivotForOrbit(FVector& OutPivot) const override;
+
+	/*
+	 * focus events
+	 */
+
+	// called when we "start" this editor mode (ie switch to this tab)
+	virtual void Enter() override;
+
+	// called when we "end" this editor mode (ie switch to another tab)
+	virtual void Exit() override;
+
+	virtual bool ShouldToolStartBeAllowed(const FString& ToolIdentifier) const override;
+
+	//////////////////
+	// End of UEdMode interface
+	//////////////////
+
+
+protected:
+	virtual void BindCommands() override;
+	virtual void CreateToolkit() override;
+	virtual void OnToolStarted(UInteractiveToolManager* Manager, UInteractiveTool* Tool) override;
+	virtual void OnToolEnded(UInteractiveToolManager* Manager, UInteractiveTool* Tool) override;
+	
+	virtual void OnToolPostBuild(UInteractiveToolManager* InToolManager, EToolSide InSide, UInteractiveTool* InBuiltTool, UInteractiveToolBuilder* InToolBuilder, const FToolBuilderState& ToolState);
+
+	void OnToolsContextRender(IToolsContextRenderAPI* RenderAPI);
+
+	void FocusCameraAtCursorHotkey();
+
+	void AcceptActiveToolActionOrTool();
+	void CancelActiveToolActionOrTool();
+
+	void ConfigureRealTimeViewportsOverride(bool bEnable);
+
+	void OnBlueprintPreCompile(UBlueprint* Blueprint);
+	FDelegateHandle BlueprintPreCompileHandle;
+
+protected:
+	UPROPERTY()
+	TObjectPtr<UScriptableToolSet> ScriptableTools;
+public:
+	virtual UScriptableToolSet* GetActiveScriptableTools() { return ScriptableTools; }
+
+};
diff --git a/Engine/Plugins/Experimental/ScriptableToolsEditorMode/Source/ScriptableToolsEditorMode/Public/ScriptableToolsEditorModeManagerCommands.h b/Engine/Plugins/Experimental/ScriptableToolsEditorMode/Source/ScriptableToolsEditorMode/Public/ScriptableToolsEditorModeManagerCommands.h
new file mode 100644
--- /dev/null
+++ b/Engine/Plugins/Experimental/ScriptableToolsEditorMode/Source/ScriptableToolsEditorMode/Public/ScriptableToolsEditorModeManagerCommands.h
@@ -0,0 +1,57 @@
+// Copyright Epic Games, Inc. All Rights Reserved.
+
+#pragma once
+
+#include "CoreMinimal.h"
+#include "Framework/Commands/Commands.h"
+
+/**
+ * TInteractiveToolCommands implementation for Modeling Mode Tools
+ */
+class SCRIPTABLETOOLSEDITORMODE_API FScriptableToolsEditorModeManagerCommands : public TCommands<FScriptableToolsEditorModeManagerCommands>
+{
+public:
+	FScriptableToolsEditorModeManagerCommands();
+
+protected:
+	struct FStartToolCommand
+	{
+		FString ToolUIName;
+		TSharedPtr<FUICommandInfo> ToolCommand;
+	};
+	TArray<FStartToolCommand> RegisteredTools;		// Tool start-commands listed below are stored in this list
+
+public:
+	/**
+	 * Find Tool start-command below by registered name (tool icon name in Mode palette)
+	 */
+	TSharedPtr<FUICommandInfo> FindToolByName(FString Name, bool& bFound) const;
+	
+	//
+	// These commands are set up to launch registered Tools via the ToolManager in ScriptableToolsEditorMode.cpp
+	//
+	TSharedPtr<FUICommandInfo> BeginMeshInspectorTool;
+	TSharedPtr<FUICommandInfo> BeginPhysicsInspectorTool;
+	TSharedPtr<FUICommandInfo> BeginSetCollisionGeometryTool;
+
+
+	TSharedPtr<FUICommandInfo> BeginMeshToPropTool;
+
+
+	//
+	// Accept/Cancel/Complete commands are used to end the active Tool via ToolManager
+	//
+
+	TSharedPtr<FUICommandInfo> AcceptActiveTool;
+	TSharedPtr<FUICommandInfo> CancelActiveTool;
+	TSharedPtr<FUICommandInfo> CompleteActiveTool;
+
+	TSharedPtr<FUICommandInfo> AcceptOrCompleteActiveTool;
+	TSharedPtr<FUICommandInfo> CancelOrCompleteActiveTool;
+
+	/**
+	 * Initialize above commands
+	 */
+	virtual void RegisterCommands() override;
+};
+
diff --git a/Engine/Plugins/Experimental/ScriptableToolsEditorMode/Source/ScriptableToolsEditorMode/Public/ScriptableToolsEditorModeModule.h b/Engine/Plugins/Experimental/ScriptableToolsEditorMode/Source/ScriptableToolsEditorMode/Public/ScriptableToolsEditorModeModule.h
new file mode 100644
--- /dev/null
+++ b/Engine/Plugins/Experimental/ScriptableToolsEditorMode/Source/ScriptableToolsEditorMode/Public/ScriptableToolsEditorModeModule.h
@@ -0,0 +1,20 @@
+// Copyright Epic Games, Inc. All Rights Reserved.
+
+#pragma once
+
+#include "CoreMinimal.h"
+#include "Modules/ModuleManager.h"
+
+class FScriptableToolsEditorModeModule : public IModuleInterface
+{
+public:
+
+	/** IModuleInterface implementation */
+	virtual void StartupModule() override;
+	virtual void ShutdownModule() override;
+
+protected:
+	void OnPostEngineInit();
+};
+
+
diff --git a/Engine/Plugins/Experimental/ScriptableToolsEditorMode/Source/ScriptableToolsEditorMode/Public/ScriptableToolsEditorModeStyle.h b/Engine/Plugins/Experimental/ScriptableToolsEditorMode/Source/ScriptableToolsEditorMode/Public/ScriptableToolsEditorModeStyle.h
new file mode 100644
--- /dev/null
+++ b/Engine/Plugins/Experimental/ScriptableToolsEditorMode/Source/ScriptableToolsEditorMode/Public/ScriptableToolsEditorModeStyle.h
@@ -0,0 +1,28 @@
+// Copyright Epic Games, Inc. All Rights Reserved.
+
+#pragma once
+
+#include "CoreMinimal.h"
+#include "Styling/ISlateStyle.h"
+#include "Styling/SlateStyle.h"
+
+class SCRIPTABLETOOLSEDITORMODE_API FScriptableToolsEditorModeStyle
+{
+public:
+	static void Initialize();
+
+	static void Shutdown();
+
+	static TSharedPtr< class ISlateStyle > Get();
+
+	static FName GetStyleSetName();
+
+	// use to access icons defined by the style set by name, eg GetBrush("BrushFalloffIcons.Smooth")
+	static const FSlateBrush* GetBrush(FName PropertyName, const ANSICHAR* Specifier = NULL);
+
+private:
+	static FString InContent(const FString& RelativePath, const ANSICHAR* Extension);
+
+private:
+	static TSharedPtr< class FSlateStyleSet > StyleSet;
+};
diff --git a/Engine/Plugins/Experimental/ScriptableToolsEditorMode/Source/ScriptableToolsEditorMode/Public/ScriptableToolsEditorModeToolkit.h b/Engine/Plugins/Experimental/ScriptableToolsEditorMode/Source/ScriptableToolsEditorMode/Public/ScriptableToolsEditorModeToolkit.h
new file mode 100644
--- /dev/null
+++ b/Engine/Plugins/Experimental/ScriptableToolsEditorMode/Source/ScriptableToolsEditorMode/Public/ScriptableToolsEditorModeToolkit.h
@@ -0,0 +1,101 @@
+// Copyright Epic Games, Inc. All Rights Reserved.
+
+#pragma once
+
+#include "CoreMinimal.h"
+#include "Toolkits/BaseToolkit.h"
+#include "ScriptableToolsEditorMode.h"
+#include "Widgets/Input/STextComboBox.h"
+#include "Widgets/Input/SEditableTextBox.h"
+#include "StatusBarSubsystem.h"
+
+
+class IDetailsView;
+class SButton;
+class STextBlock;
+class UBlueprint;
+
+class SCRIPTABLETOOLSEDITORMODE_API FScriptableToolsEditorModeToolkit : public FModeToolkit
+{
+public:
+
+	FScriptableToolsEditorModeToolkit();
+	~FScriptableToolsEditorModeToolkit();
+	
+	/** FModeToolkit interface */
+	virtual void Init(const TSharedPtr<IToolkitHost>& InitToolkitHost, TWeakObjectPtr<UEdMode> InOwningMode) override;
+
+	/** IToolkit interface */
+	virtual FName GetToolkitFName() const override;
+	virtual FText GetBaseToolkitName() const override;
+	virtual TSharedPtr<class SWidget> GetInlineContent() const override { return ToolkitWidget; }
+
+	// initialize toolkit widgets that need to wait until mode is initialized/entered
+	virtual void InitializeAfterModeSetup();
+
+	// set/clear notification message area
+	virtual void PostNotification(const FText& Message);
+	virtual void ClearNotification();
+
+	// set/clear warning message area
+	virtual void PostWarning(const FText& Message);
+	virtual void ClearWarning();
+
+	/** Returns the Mode specific tabs in the mode toolbar **/ 
+	virtual void GetToolPaletteNames(TArray<FName>& InPaletteName) const;
+	virtual FText GetToolPaletteDisplayName(FName PaletteName) const; 
+	virtual void BuildToolPalette(FName PaletteName, class FToolBarBuilder& ToolbarBuilder);
+	virtual void OnToolPaletteChanged(FName PaletteName) override;
+	virtual bool HasIntegratedToolPalettes() const { return false; }
+	virtual bool HasExclusiveToolPalettes() const { return false; }
+
+	virtual FText GetActiveToolDisplayName() const override { return ActiveToolName; }
+	virtual FText GetActiveToolMessage() const override { return ActiveToolMessage; }
+
+	virtual void EnableShowRealtimeWarning(bool bEnable);
+
+	virtual void OnToolStarted(UInteractiveToolManager* Manager, UInteractiveTool* Tool) override;
+	virtual void OnToolEnded(UInteractiveToolManager* Manager, UInteractiveTool* Tool) override;
+
+	virtual void CustomizeModeDetailsViewArgs(FDetailsViewArgs& ArgsInOut) override;
+
+	void OnActiveViewportChanged(TSharedPtr<IAssetViewport>, TSharedPtr<IAssetViewport> );
+
+	virtual void InvokeUI() override;
+
+	virtual void ForceToolPaletteRebuild();
+
+private:
+	const static TArray<FName> PaletteNames_Standard;
+
+	FText ActiveToolName;
+	FText ActiveToolMessage;
+	FStatusBarMessageHandle ActiveToolMessageHandle;
+	const FSlateBrush* ActiveToolIcon = nullptr;
+
+	TSharedPtr<SWidget> ToolkitWidget;
+	void UpdateActiveToolProperties();
+	void InvalidateCachedDetailPanelState(UObject* ChangedObject);
+
+	TSharedPtr<SWidget> ViewportOverlayWidget;
+
+	TSharedPtr<STextBlock> ModeWarningArea;
+	TSharedPtr<STextBlock> ModeHeaderArea;
+	TSharedPtr<STextBlock> ToolWarningArea;
+	TSharedPtr<SButton> AcceptButton;
+	TSharedPtr<SButton> CancelButton;
+	TSharedPtr<SButton> CompletedButton;
+
+	bool bShowRealtimeWarning = false;
+	void UpdateShowWarnings();
+
+	TMap<FName, FText> ActiveToolCategories;
+	void UpdateActiveToolCategories();
+
+	bool bFirstInitializeAfterModeSetup = true;
+
+	bool bShowActiveSelectionActions = false;
+
+
+	// custom accept/cancel/complete handlers
+};
diff --git a/Engine/Plugins/Experimental/ScriptableToolsEditorMode/Source/ScriptableToolsEditorMode/ScriptableToolsEditorMode.Build.cs b/Engine/Plugins/Experimental/ScriptableToolsEditorMode/Source/ScriptableToolsEditorMode/ScriptableToolsEditorMode.Build.cs
new file mode 100644
--- /dev/null
+++ b/Engine/Plugins/Experimental/ScriptableToolsEditorMode/Source/ScriptableToolsEditorMode/ScriptableToolsEditorMode.Build.cs
@@ -0,0 +1,73 @@
+// Copyright Epic Games, Inc. All Rights Reserved.
+
+using UnrealBuildTool;
+
+public class ScriptableToolsEditorMode : ModuleRules
+{
+	public ScriptableToolsEditorMode(ReadOnlyTargetRules Target) : base(Target)
+	{
+		PCHUsage = ModuleRules.PCHUsageMode.UseExplicitOrSharedPCHs;
+
+		PublicIncludePaths.AddRange(
+			new string[] {
+				// ... add public include paths required here ...
+			}
+			);
+
+
+		PrivateIncludePaths.AddRange(
+			new string[] {
+				// ... add other private include paths required here ...
+			}
+			);
+
+
+		PublicDependencyModuleNames.AddRange(
+			new string[]
+			{
+				"CoreUObject",
+				"Slate",
+				"SlateCore",
+				"Engine",
+				"InputCore",
+				"UnrealEd",
+				"EditorFramework",
+				"ToolWidgets",
+				"EditorWidgets",
+				"DeveloperSettings",
+				"PropertyEditor",
+				"StatusBar",
+				"Projects",
+
+				"Core",
+				"PhysicsCore",
+				"RenderCore",
+
+				"GeometryCore",
+				"InteractiveToolsFramework",
+				"EditorInteractiveToolsFramework",
+				"ScriptableToolsFramework",
+				"EditorScriptableToolsFramework",
+
+				"ModelingComponents"
+			}
+			);
+
+
+		PrivateDependencyModuleNames.AddRange(
+			new string[]
+			{
+				"Engine"
+			}
+			);
+
+
+		DynamicallyLoadedModuleNames.AddRange(
+			new string[]
+			{
+				// ... add any modules that your module loads dynamically here ...
+			}
+			);
+
+	}
+}
diff --git a/Engine/Plugins/Experimental/ScriptableToolsFramework/ScriptableToolsFramework.uplugin b/Engine/Plugins/Experimental/ScriptableToolsFramework/ScriptableToolsFramework.uplugin
new file mode 100644
--- /dev/null
+++ b/Engine/Plugins/Experimental/ScriptableToolsFramework/ScriptableToolsFramework.uplugin
@@ -0,0 +1,30 @@
+{
+	"FileVersion": 3,
+	"Version": 1,
+	"VersionName": "0.1",
+	"FriendlyName": "Scriptable Tools Framework",
+	"Description": "Blueprint-Scriptable extension to the Interactive Tools Framework",
+	"Category": "Editor",
+	"CreatedBy": "Epic Games, Inc.",
+	"CreatedByURL": "http://epicgames.com",
+	"DocsURL": "",
+	"MarketplaceURL": "",
+	"SupportURL": "",
+	"CanContainContent": true,
+	"IsExperimentalVersion": true,
+	"Installed": false,
+	"Modules": [
+		{
+			"Name": "ScriptableToolsFramework",
+			"Type": "Runtime",
+			"LoadingPhase": "Default"
+		},
+		{
+			"Name": "EditorScriptableToolsFramework",
+			"Type": "Editor",
+			"LoadingPhase": "Default"
+		}		
+	],
+	"Plugins": [
+	]
+}
diff --git a/Engine/Plugins/Experimental/ScriptableToolsFramework/Source/EditorScriptableToolsFramework/EditorScriptableToolsFramework.Build.cs b/Engine/Plugins/Experimental/ScriptableToolsFramework/Source/EditorScriptableToolsFramework/EditorScriptableToolsFramework.Build.cs
new file mode 100644
--- /dev/null
+++ b/Engine/Plugins/Experimental/ScriptableToolsFramework/Source/EditorScriptableToolsFramework/EditorScriptableToolsFramework.Build.cs
@@ -0,0 +1,55 @@
+// Copyright Epic Games, Inc. All Rights Reserved.
+
+using UnrealBuildTool;
+
+public class EditorScriptableToolsFramework : ModuleRules
+{
+	public EditorScriptableToolsFramework(ReadOnlyTargetRules Target) : base(Target)
+	{
+		PCHUsage = ModuleRules.PCHUsageMode.UseExplicitOrSharedPCHs;
+
+		PublicIncludePaths.AddRange(
+			new string[] {
+				// ... add public include paths required here ...
+			}
+			);
+
+
+		PrivateIncludePaths.AddRange(
+			new string[] {
+				// ... add other private include paths required here ...
+			}
+			);
+
+
+		PublicDependencyModuleNames.AddRange(
+			new string[]
+			{
+				"Core",
+				"CoreUObject",
+				"PhysicsCore",
+				"RenderCore",
+				"GeometryCore",
+				"InteractiveToolsFramework",
+				"ScriptableToolsFramework"
+			}
+			);
+
+
+		PrivateDependencyModuleNames.AddRange(
+			new string[]
+			{
+				"Engine"
+			}
+			);
+
+
+		DynamicallyLoadedModuleNames.AddRange(
+			new string[]
+			{
+				// ... add any modules that your module loads dynamically here ...
+			}
+			);
+
+	}
+}
diff --git a/Engine/Plugins/Experimental/ScriptableToolsFramework/Source/EditorScriptableToolsFramework/Private/EditorScriptableToolsFrameworkModule.cpp b/Engine/Plugins/Experimental/ScriptableToolsFramework/Source/EditorScriptableToolsFramework/Private/EditorScriptableToolsFrameworkModule.cpp
new file mode 100644
--- /dev/null
+++ b/Engine/Plugins/Experimental/ScriptableToolsFramework/Source/EditorScriptableToolsFramework/Private/EditorScriptableToolsFrameworkModule.cpp
@@ -0,0 +1,20 @@
+// Copyright Epic Games, Inc. All Rights Reserved.
+
+#include "EditorScriptableToolsFrameworkModule.h"
+
+#define LOCTEXT_NAMESPACE "FEditorScriptableToolsFrameworkModule"
+
+void FEditorScriptableToolsFrameworkModule::StartupModule()
+{
+	// This code will execute after your module is loaded into memory; the exact timing is specified in the .uplugin file per-module
+}
+
+void FEditorScriptableToolsFrameworkModule::ShutdownModule()
+{
+	// This function may be called during shutdown to clean up your module.  For modules that support dynamic reloading,
+	// we call this function before unloading the module.
+}
+
+#undef LOCTEXT_NAMESPACE
+	
+IMPLEMENT_MODULE(FEditorScriptableToolsFrameworkModule, EditorScriptableToolsFramework)
diff --git a/Engine/Plugins/Experimental/ScriptableToolsFramework/Source/EditorScriptableToolsFramework/Public/BaseTools/EditorScriptableClickDragTool.h b/Engine/Plugins/Experimental/ScriptableToolsFramework/Source/EditorScriptableToolsFramework/Public/BaseTools/EditorScriptableClickDragTool.h
new file mode 100644
--- /dev/null
+++ b/Engine/Plugins/Experimental/ScriptableToolsFramework/Source/EditorScriptableToolsFramework/Public/BaseTools/EditorScriptableClickDragTool.h
@@ -0,0 +1,17 @@
+// Copyright Epic Games, Inc. All Rights Reserved.
+
+#pragma once
+
+#include "BaseTools/ScriptableClickDragTool.h"
+#include "EditorScriptableClickDragTool.generated.h"
+
+/**
+ * Editor-Only variant of UScriptableClickDragTool, which gives access to Editor-Only BP functions
+ */
+UCLASS(Transient, Blueprintable)
+class EDITORSCRIPTABLETOOLSFRAMEWORK_API UEditorScriptableClickDragTool : public UScriptableClickDragTool
+{
+	GENERATED_BODY()
+public:
+
+};
\ No newline at end of file
diff --git a/Engine/Plugins/Experimental/ScriptableToolsFramework/Source/EditorScriptableToolsFramework/Public/BaseTools/EditorScriptableSingleClickTool.h b/Engine/Plugins/Experimental/ScriptableToolsFramework/Source/EditorScriptableToolsFramework/Public/BaseTools/EditorScriptableSingleClickTool.h
new file mode 100644
--- /dev/null
+++ b/Engine/Plugins/Experimental/ScriptableToolsFramework/Source/EditorScriptableToolsFramework/Public/BaseTools/EditorScriptableSingleClickTool.h
@@ -0,0 +1,17 @@
+// Copyright Epic Games, Inc. All Rights Reserved.
+
+#pragma once
+
+#include "BaseTools/ScriptableSingleClickTool.h"
+#include "EditorScriptableSingleClickTool.generated.h"
+
+/**
+ * Editor-Only variant of UScriptableSingleClickTool, which gives access to Editor-Only BP functions
+ */
+UCLASS(Transient, Blueprintable)
+class EDITORSCRIPTABLETOOLSFRAMEWORK_API UEditorScriptableSingleClickTool : public UScriptableSingleClickTool
+{
+	GENERATED_BODY()
+public:
+
+};
\ No newline at end of file
diff --git a/Engine/Plugins/Experimental/ScriptableToolsFramework/Source/EditorScriptableToolsFramework/Public/EditorScriptableInteractiveTool.h b/Engine/Plugins/Experimental/ScriptableToolsFramework/Source/EditorScriptableToolsFramework/Public/EditorScriptableInteractiveTool.h
new file mode 100644
--- /dev/null
+++ b/Engine/Plugins/Experimental/ScriptableToolsFramework/Source/EditorScriptableToolsFramework/Public/EditorScriptableInteractiveTool.h
@@ -0,0 +1,17 @@
+// Copyright Epic Games, Inc. All Rights Reserved.
+
+#pragma once
+
+#include "ScriptableInteractiveTool.h"
+#include "EditorScriptableInteractiveTool.generated.h"
+
+/**
+ * Editor-Only variant of UScriptableInteractiveTool, which gives access to Editor-Only BP functions
+ */
+UCLASS(Transient, Blueprintable)
+class EDITORSCRIPTABLETOOLSFRAMEWORK_API UEditorScriptableInteractiveTool : public UScriptableInteractiveTool
+{
+	GENERATED_BODY()
+public:
+
+};
\ No newline at end of file
diff --git a/Engine/Plugins/Experimental/ScriptableToolsFramework/Source/EditorScriptableToolsFramework/Public/EditorScriptableToolsFrameworkModule.h b/Engine/Plugins/Experimental/ScriptableToolsFramework/Source/EditorScriptableToolsFramework/Public/EditorScriptableToolsFrameworkModule.h
new file mode 100644
--- /dev/null
+++ b/Engine/Plugins/Experimental/ScriptableToolsFramework/Source/EditorScriptableToolsFramework/Public/EditorScriptableToolsFrameworkModule.h
@@ -0,0 +1,17 @@
+// Copyright Epic Games, Inc. All Rights Reserved.
+
+#pragma once
+
+#include "CoreMinimal.h"
+#include "Modules/ModuleManager.h"
+
+class FEditorScriptableToolsFrameworkModule : public IModuleInterface
+{
+public:
+
+	/** IModuleInterface implementation */
+	virtual void StartupModule() override;
+	virtual void ShutdownModule() override;
+};
+
+
diff --git a/Engine/Plugins/Experimental/ScriptableToolsFramework/Source/ScriptableToolsFramework/Private/BaseTools/ScriptableClickDragTool.cpp b/Engine/Plugins/Experimental/ScriptableToolsFramework/Source/ScriptableToolsFramework/Private/BaseTools/ScriptableClickDragTool.cpp
new file mode 100644
--- /dev/null
+++ b/Engine/Plugins/Experimental/ScriptableToolsFramework/Source/ScriptableToolsFramework/Private/BaseTools/ScriptableClickDragTool.cpp
@@ -0,0 +1,189 @@
+// Copyright Epic Games, Inc. All Rights Reserved.
+
+#include "BaseTools/ScriptableClickDragTool.h"
+#include "BaseBehaviors/ClickDragBehavior.h"
+#include "BaseBehaviors/MouseHoverBehavior.h"
+
+#define LOCTEXT_NAMESPACE "UScriptableClickDragTool"
+
+void UScriptableClickDragTool::Setup()
+{
+	UScriptableInteractiveTool::Setup();
+
+	// add default mouse input behavior
+	ClickDragBehavior = NewObject<UClickDragInputBehavior>();
+	ClickDragBehavior->Initialize(this);
+	ClickDragBehavior->Modifiers.RegisterModifier(1, FInputDeviceState::IsShiftKeyDown);
+	ClickDragBehavior->Modifiers.RegisterModifier(2, FInputDeviceState::IsCtrlKeyDown);
+	ClickDragBehavior->Modifiers.RegisterModifier(3, FInputDeviceState::IsAltKeyDown);
+	ClickDragBehavior->bUpdateModifiersDuringDrag = this->bUpdateModifiersDuringDrag;
+	AddInputBehavior(ClickDragBehavior);
+
+	bShiftModifier = false;
+	bCtrlModifier = false;
+	bAltModifier = false;
+
+	if (bWantMouseHover)
+	{
+		MouseHoverBehavior = NewObject<UMouseHoverBehavior>();
+		MouseHoverBehavior->Initialize(this);
+		MouseHoverBehavior->Modifiers.RegisterModifier(1, FInputDeviceState::IsShiftKeyDown);
+		MouseHoverBehavior->Modifiers.RegisterModifier(2, FInputDeviceState::IsCtrlKeyDown);
+		MouseHoverBehavior->Modifiers.RegisterModifier(3, FInputDeviceState::IsAltKeyDown);
+		AddInputBehavior(MouseHoverBehavior);
+	}
+}
+
+
+FInputRayHit UScriptableClickDragTool::CanBeginClickDragSequence(const FInputDeviceRay& PressPos)
+{
+	return TestIfCanBeginClickDrag(PressPos, GetActiveModifiers());
+}
+
+void UScriptableClickDragTool::OnClickPress(const FInputDeviceRay& PressPos)
+{
+	bInClickDrag = true;
+	OnDragBegin(PressPos, GetActiveModifiers());
+}
+
+void UScriptableClickDragTool::OnClickDrag(const FInputDeviceRay& DragPos)
+{
+	OnDragUpdatePosition(DragPos, GetActiveModifiers());
+}
+
+void UScriptableClickDragTool::OnClickRelease(const FInputDeviceRay& ReleasePos)
+{
+	OnDragEnd(ReleasePos, GetActiveModifiers());
+	bInClickDrag = false;
+}
+
+void UScriptableClickDragTool::OnTerminateDragSequence()
+{
+	OnDragSequenceCancelled();
+	bInClickDrag = false;
+}
+
+
+bool UScriptableClickDragTool::InActiveClickDrag() const
+{
+	return bInClickDrag;
+}
+
+
+FInputRayHit UScriptableClickDragTool::TestIfCanBeginClickDrag_Implementation(FInputDeviceRay ClickPos, const FScriptableToolModifierStates& Modifiers)
+{
+	// always return a hit
+	return FInputRayHit(0.0f);
+}
+
+void UScriptableClickDragTool::OnDragBegin_Implementation(FInputDeviceRay ClickPos, const FScriptableToolModifierStates& Modifiers)
+{
+}
+
+void UScriptableClickDragTool::OnDragUpdatePosition_Implementation(FInputDeviceRay ClickPos, const FScriptableToolModifierStates& Modifiers)
+{
+}
+
+void UScriptableClickDragTool::OnDragEnd_Implementation(FInputDeviceRay ClickPos, const FScriptableToolModifierStates& Modifiers)
+{
+}
+
+void UScriptableClickDragTool::OnDragSequenceCancelled_Implementation()
+{
+}
+
+
+
+
+void UScriptableClickDragTool::OnUpdateModifierState(int ModifierID, bool bIsOn)
+{
+	if (ModifierID == 1)
+	{
+		bShiftModifier = bIsOn;
+	}
+	if (ModifierID == 2)
+	{
+		bCtrlModifier = bIsOn;
+	}
+	if (ModifierID == 3)
+	{
+		bAltModifier = bIsOn;
+	}
+}
+
+bool UScriptableClickDragTool::IsShiftDown() const
+{
+	return bShiftModifier;
+}
+
+bool UScriptableClickDragTool::IsCtrlDown() const
+{
+	return bCtrlModifier;
+}
+
+bool UScriptableClickDragTool::IsAltDown() const
+{
+	return bAltModifier;
+}
+
+FScriptableToolModifierStates UScriptableClickDragTool::GetActiveModifiers()
+{
+	FScriptableToolModifierStates Modifiers;
+	Modifiers.bShiftDown = bShiftModifier;
+	Modifiers.bCtrlDown = bCtrlModifier;
+	Modifiers.bAltDown = bAltModifier;
+	return Modifiers;
+}
+
+
+
+
+FInputRayHit UScriptableClickDragTool::OnHoverHitTest_Implementation(FInputDeviceRay ClickPos, const FScriptableToolModifierStates& Modifiers)
+{
+	// always return a hit
+	return FInputRayHit(0.0f);
+}
+FInputRayHit UScriptableClickDragTool::BeginHoverSequenceHitTest(const FInputDeviceRay& ClickPos)
+{
+	return OnHoverHitTest(ClickPos, GetActiveModifiers());
+}
+
+void UScriptableClickDragTool::OnHoverBegin_Implementation(FInputDeviceRay HoverPos, const FScriptableToolModifierStates& Modifiers)
+{
+}
+
+void UScriptableClickDragTool::OnBeginHover(const FInputDeviceRay& DevicePos)
+{
+	bInHover = true;
+	OnHoverBegin(DevicePos, GetActiveModifiers());
+}
+
+bool UScriptableClickDragTool::OnHoverUpdate_Implementation(FInputDeviceRay HoverPos, const FScriptableToolModifierStates& Modifiers)
+{
+	return true;
+}
+
+bool UScriptableClickDragTool::OnUpdateHover(const FInputDeviceRay& DevicePos)
+{
+	return OnHoverUpdate(DevicePos, GetActiveModifiers());
+}
+
+
+void UScriptableClickDragTool::OnHoverEnd_Implementation(const FScriptableToolModifierStates& Modifiers)
+{
+}
+
+void UScriptableClickDragTool::OnEndHover()
+{
+	OnHoverEnd(GetActiveModifiers());
+	bInHover = false;
+}
+
+bool UScriptableClickDragTool::InActiveHover() const
+{
+	return bInHover;
+}
+
+
+
+#undef LOCTEXT_NAMESPACE
diff --git a/Engine/Plugins/Experimental/ScriptableToolsFramework/Source/ScriptableToolsFramework/Private/BaseTools/ScriptableSingleClickTool.cpp b/Engine/Plugins/Experimental/ScriptableToolsFramework/Source/ScriptableToolsFramework/Private/BaseTools/ScriptableSingleClickTool.cpp
new file mode 100644
--- /dev/null
+++ b/Engine/Plugins/Experimental/ScriptableToolsFramework/Source/ScriptableToolsFramework/Private/BaseTools/ScriptableSingleClickTool.cpp
@@ -0,0 +1,147 @@
+// Copyright Epic Games, Inc. All Rights Reserved.
+
+#include "BaseTools/ScriptableSingleClickTool.h"
+#include "BaseBehaviors/SingleClickBehavior.h"
+#include "BaseBehaviors/MouseHoverBehavior.h"
+
+
+#define LOCTEXT_NAMESPACE "UScriptableSingleClickTool"
+
+void UScriptableSingleClickTool::Setup()
+{
+	UScriptableInteractiveTool::Setup();
+
+	// add default button input behaviors for devices
+	SingleClickBehavior = NewObject<USingleClickInputBehavior>();
+	SingleClickBehavior->Initialize(this);
+	AddInputBehavior(SingleClickBehavior);
+
+	SingleClickBehavior->Modifiers.RegisterModifier(1, FInputDeviceState::IsShiftKeyDown);
+	SingleClickBehavior->Modifiers.RegisterModifier(2, FInputDeviceState::IsCtrlKeyDown);
+	SingleClickBehavior->Modifiers.RegisterModifier(3, FInputDeviceState::IsAltKeyDown);
+	bShiftModifier = false;
+	bCtrlModifier = false;
+	bAltModifier = false;
+
+
+	if (bWantMouseHover)
+	{
+		MouseHoverBehavior = NewObject<UMouseHoverBehavior>();
+		MouseHoverBehavior->Initialize(this);
+		AddInputBehavior(MouseHoverBehavior);
+	}
+
+}
+
+
+void UScriptableSingleClickTool::OnUpdateModifierState(int ModifierID, bool bIsOn)
+{
+	if (ModifierID == 1)
+	{
+		bShiftModifier = bIsOn;
+	}
+	if (ModifierID == 2)
+	{
+		bCtrlModifier = bIsOn;
+	}
+	if (ModifierID == 3)
+	{
+		bAltModifier = bIsOn;
+	}
+}
+
+bool UScriptableSingleClickTool::IsShiftDown() const
+{
+	return bShiftModifier;
+}
+
+bool UScriptableSingleClickTool::IsCtrlDown() const
+{
+	return bCtrlModifier;
+}
+
+bool UScriptableSingleClickTool::IsAltDown() const
+{
+	return bAltModifier;
+}
+
+FScriptableToolModifierStates UScriptableSingleClickTool::GetActiveModifiers()
+{
+	FScriptableToolModifierStates Modifiers;
+	Modifiers.bShiftDown = bShiftModifier;
+	Modifiers.bCtrlDown = bCtrlModifier;
+	Modifiers.bAltDown = bAltModifier;
+	return Modifiers;
+}
+
+
+FInputRayHit UScriptableSingleClickTool::TestIfHitByClick_Implementation(FInputDeviceRay ClickPos, const FScriptableToolModifierStates& Modifiers)
+{
+	UE_LOG(LogTemp, Warning, TEXT("UScriptableSingleClickTool: Default TestIfHitByClick Implementation always consumes hit"));
+
+	// always return a hit
+	return FInputRayHit(0.0f);
+}
+FInputRayHit UScriptableSingleClickTool::IsHitByClick(const FInputDeviceRay& ClickPos)
+{
+	return TestIfHitByClick(ClickPos, GetActiveModifiers());
+}
+
+void UScriptableSingleClickTool::OnHitByClick_Implementation(FInputDeviceRay ClickPos, const FScriptableToolModifierStates& Modifiers)
+{
+	UE_LOG(LogTemp, Warning, TEXT("UScriptableSingleClickTool: Empty OnHitByClick Implementation"));
+}
+void UScriptableSingleClickTool::OnClicked(const FInputDeviceRay& ClickPos)
+{
+	OnHitByClick(ClickPos, GetActiveModifiers());
+}
+
+
+FInputRayHit UScriptableSingleClickTool::OnHoverHitTest_Implementation(FInputDeviceRay ClickPos, const FScriptableToolModifierStates& Modifiers)
+{
+	// always return a hit
+	return FInputRayHit(0.0f);
+}
+FInputRayHit UScriptableSingleClickTool::BeginHoverSequenceHitTest(const FInputDeviceRay& ClickPos)
+{
+	return OnHoverHitTest(ClickPos, GetActiveModifiers());
+}
+
+void UScriptableSingleClickTool::OnHoverBegin_Implementation(FInputDeviceRay HoverPos, const FScriptableToolModifierStates& Modifiers)
+{
+}
+
+void UScriptableSingleClickTool::OnBeginHover(const FInputDeviceRay& DevicePos)
+{
+	bInHover = true;
+	OnHoverBegin(DevicePos, GetActiveModifiers());
+}
+
+bool UScriptableSingleClickTool::OnHoverUpdate_Implementation(FInputDeviceRay HoverPos, const FScriptableToolModifierStates& Modifiers)
+{
+	return true;
+}
+
+bool UScriptableSingleClickTool::OnUpdateHover(const FInputDeviceRay& DevicePos)
+{
+	return OnHoverUpdate(DevicePos, GetActiveModifiers());
+}
+
+
+void UScriptableSingleClickTool::OnHoverEnd_Implementation(const FScriptableToolModifierStates& Modifiers)
+{
+}
+
+void UScriptableSingleClickTool::OnEndHover()
+{
+	OnHoverEnd(GetActiveModifiers());
+	bInHover = false;
+}
+
+bool UScriptableSingleClickTool::InActiveHover() const
+{
+	return bInHover;
+}
+
+
+#undef LOCTEXT_NAMESPACE
diff --git a/Engine/Plugins/Experimental/ScriptableToolsFramework/Source/ScriptableToolsFramework/Private/ScriptableInteractiveTool.cpp b/Engine/Plugins/Experimental/ScriptableToolsFramework/Source/ScriptableToolsFramework/Private/ScriptableInteractiveTool.cpp
new file mode 100644
--- /dev/null
+++ b/Engine/Plugins/Experimental/ScriptableToolsFramework/Source/ScriptableToolsFramework/Private/ScriptableInteractiveTool.cpp
@@ -0,0 +1,1046 @@
+// Copyright Epic Games, Inc. All Rights Reserved.
+
+
+#include "ScriptableInteractiveTool.h"
+#include "InteractiveToolManager.h"
+#include "UObject/Class.h"
+#include "Engine/World.h"
+#include "Misc/EnumClassFlags.h"
+
+#include "BaseGizmos/TransformGizmoUtil.h"
+#include "BaseGizmos/CombinedTransformGizmo.h"
+
+#include "ToolDataVisualizer.h"
+#include "CanvasTypes.h"
+#include "CanvasItem.h"
+#include "Engine/Engine.h"  // for GEngine->GetSmallFont()
+#include "SceneView.h"
+
+#include "Serialization/MemoryHasher.h"
+
+#define LOCTEXT_NAMESPACE "UScriptableInteractiveTool"
+
+
+
+UScriptableInteractiveTool* UScriptableInteractiveToolPropertySet::GetOwningTool(TEnumAsByte<EToolsFrameworkOutcomePins>& Outcome)
+{
+	if (ParentTool.IsValid())
+	{
+		Outcome = EToolsFrameworkOutcomePins::Success;
+		return ParentTool.Get();
+	}
+	Outcome = EToolsFrameworkOutcomePins::Failure;
+	return nullptr;
+}
+
+
+
+
+UScriptableTool_RenderAPI* UScriptableTool_RenderAPI::DrawLine(FVector Start, FVector End, FLinearColor Color, float Thickness, float DepthBias, bool bDepthTested)
+{
+	if (ActiveVisualizer == nullptr)
+	{
+		UE_LOG(LogTemp, Warning, TEXT("[ScriptableInteractiveTool::DrawLine] DrawLine can only be called during OnScriptRender event callback"));
+		return this;
+	}
+
+	ActiveVisualizer->DepthBias = DepthBias;
+	ActiveVisualizer->DrawLine(Start, End, Color, Thickness, bDepthTested);
+	return this;
+}
+
+
+UScriptableTool_HUDAPI* UScriptableTool_HUDAPI::DrawTextAtLocation(FVector Location, FString String, FLinearColor Color, bool bCentered, float ShiftRowsY)
+{
+	if (ActiveCanvas == nullptr || ActiveSceneView == nullptr)
+	{
+		UE_LOG(LogTemp, Warning, TEXT("[ScriptableInteractiveTool::DrawTextAtLocation] DrawTextAtLocation can only be called during OnScriptDrawHUD event callback"));
+		return this;
+	}
+
+	float DPIScale = ActiveCanvas->GetDPIScale();
+	UFont* UseFont = GEngine->GetSmallFont();
+
+	FVector2D PixelPos;
+	ActiveSceneView->WorldToPixel(Location, PixelPos);
+
+	double Height = (double)UseFont->GetStringHeightSize(TEXT("yjMT"));
+	PixelPos.Y -= Height;
+	PixelPos.Y -= ShiftRowsY * Height;
+
+	if ( bCentered )
+	{
+		float StringWidth = UseFont->GetStringSize(*String);
+		PixelPos.X -= 0.5*StringWidth;
+	}
+
+	ActiveCanvas->DrawShadowedString(PixelPos.X / (double)DPIScale, PixelPos.Y / (double)DPIScale, *String, UseFont, Color);
+
+	return this;
+}
+
+
+UScriptableTool_HUDAPI* UScriptableTool_HUDAPI::DrawTextArrayAtLocation(FVector Location, TArray<FString> Strings, FLinearColor Color, bool bCentered, float ShiftRowsY)
+{
+	if (ActiveCanvas == nullptr || ActiveSceneView == nullptr)
+	{
+		UE_LOG(LogTemp, Warning, TEXT("[ScriptableInteractiveTool::DrawTextAtLocation] DrawTextAtLocation can only be called during OnScriptDrawHUD event callback"));
+		return this;
+	}
+
+	float DPIScale = ActiveCanvas->GetDPIScale();
+	UFont* UseFont = GEngine->GetSmallFont();
+
+	FVector2D PixelPos;
+	ActiveSceneView->WorldToPixel(Location, PixelPos);
+
+	double Height = (double)UseFont->GetStringHeightSize(TEXT("yjMT"));
+	double CurY = PixelPos.Y;
+	CurY -= Height;
+	CurY -= ShiftRowsY * Height;
+
+	int32 NumStrings = Strings.Num();
+	if (NumStrings > 1)
+	{
+		Height *= 1.5;
+	}
+
+	for (int32 k = 0; k < NumStrings; ++k)
+	{
+		const FString& CurString = Strings[NumStrings-k-1];
+		
+		double ShiftX = 0;
+		if (bCentered)
+		{
+			ShiftX = -0.5 * UseFont->GetStringSize(*CurString);
+		}
+
+		ActiveCanvas->DrawShadowedString((PixelPos.X + ShiftX) / (double)DPIScale, CurY / (double)DPIScale, *CurString, UseFont, Color);
+		CurY -= Height;
+	}
+
+	return this;
+}
+
+
+UScriptableTool_HUDAPI* UScriptableTool_HUDAPI::GetCanvasLocation(FVector Location, FVector2D& CanvasLocation)
+{
+	if (ActiveCanvas == nullptr || ActiveSceneView == nullptr)
+	{
+		UE_LOG(LogTemp, Warning, TEXT("[ScriptableInteractiveTool::GetCanvasLocation] DrawTextAtLocation can only be called during OnScriptDrawHUD event callback"));
+		CanvasLocation = FVector2D::Zero();
+	}
+	else
+	{
+		float DPIScale = ActiveCanvas->GetDPIScale();
+		FVector2D PixelPos;
+		ActiveSceneView->WorldToPixel(Location, PixelPos);
+		CanvasLocation = PixelPos / (double)DPIScale;
+	}
+	return this;
+}
+
+
+
+void UScriptableInteractiveTool::Setup()
+{
+	UInteractiveTool::Setup();
+
+	RenderHelper = NewObject<UScriptableTool_RenderAPI>();
+	DrawHUDHelper = NewObject<UScriptableTool_HUDAPI>();
+
+	OnScriptSetup();
+}
+
+void UScriptableInteractiveTool::OnTick(float DeltaTime)
+{
+	UInteractiveTool::OnTick(DeltaTime);
+	OnScriptTick(DeltaTime);
+}
+
+bool UScriptableInteractiveTool::HasAccept() const
+{
+	return (ToolShutdownType == EScriptableToolShutdownType::AcceptCancel);
+}
+bool UScriptableInteractiveTool::HasCancel() const
+{
+	return (ToolShutdownType == EScriptableToolShutdownType::AcceptCancel);
+}
+
+bool UScriptableInteractiveTool::CanAccept() const 
+{
+	if (ToolShutdownType != EScriptableToolShutdownType::AcceptCancel)
+	{
+		return false;
+	}
+	return OnScriptCanAccept();
+}
+
+bool UScriptableInteractiveTool::OnScriptCanAccept_Implementation() const
+{
+	return HasAccept();
+}
+
+
+void UScriptableInteractiveTool::Shutdown(EToolShutdownType ShutdownType)
+{
+	OnScriptShutdown(ShutdownType);
+
+	// clean up any gizmos
+	for (TPair<FString, TObjectPtr<UInteractiveGizmo>> Gizmo : Gizmos)
+	{
+		GetToolManager()->GetPairedGizmoManager()->DestroyGizmo(Gizmo.Value);
+	}
+	Gizmos.Reset();
+
+	RenderHelper = nullptr;
+	DrawHUDHelper = nullptr;
+
+	UInteractiveTool::Shutdown(ShutdownType);
+}
+
+void UScriptableInteractiveTool::RequestToolShutdown(
+	bool bAccept,
+	bool bShowUserPopupMessage,
+	FText UserMessage)
+{
+	if (UInteractiveToolManager* ToolManager = GetToolManager())
+	{
+		EToolShutdownType UseShutdownType = EToolShutdownType::Completed;
+		if (ToolShutdownType == EScriptableToolShutdownType::AcceptCancel)
+		{
+			UseShutdownType = (bAccept) ? EToolShutdownType::Accept : EToolShutdownType::Cancel;
+		}
+		ToolManager->PostActiveToolShutdownRequest(this, UseShutdownType, bShowUserPopupMessage, UserMessage);
+	}
+}
+
+
+
+void UScriptableInteractiveTool::PostInitProperties()
+{
+	Super::PostInitProperties();
+
+	// because this was done in the C++ Constructor, it will be wiped out when a BP Subclass
+	// initializes itself from it's CDO
+	if (InputBehaviors == nullptr)
+	{
+		InputBehaviors = NewObject<UInputBehaviorSet>(this, TEXT("InputBehaviors"));
+	}
+}
+
+
+
+void UScriptableInteractiveTool::SetTargetWorld(UWorld* World)
+{
+	TargetWorld = World;
+}
+
+UWorld* UScriptableInteractiveTool::GetWorld() const 
+{ 
+	return TargetWorld.Get(); 
+}
+
+
+UWorld* UScriptableInteractiveTool::GetToolWorld()
+{
+	return TargetWorld.Get();
+}
+
+
+
+
+void UScriptableInteractiveTool::Render(IToolsContextRenderAPI* RenderAPI)
+{
+	FToolDataVisualizer Renderer;
+	Renderer.BeginFrame(RenderAPI);
+
+	RenderHelper->ActiveVisualizer = &Renderer;
+
+	OnScriptRender(RenderHelper);
+
+	RenderHelper->ActiveVisualizer = nullptr;
+
+	Renderer.EndFrame();
+}
+void UScriptableInteractiveTool::DrawHUD(FCanvas* Canvas, IToolsContextRenderAPI* RenderAPI)
+{
+	DrawHUDHelper->ActiveCanvas = Canvas;
+	DrawHUDHelper->ActiveSceneView = RenderAPI->GetSceneView();
+
+	OnScriptDrawHUD(DrawHUDHelper);
+
+	DrawHUDHelper->ActiveCanvas = nullptr;
+	DrawHUDHelper->ActiveSceneView = nullptr;
+
+}
+
+
+UScriptableInteractiveToolPropertySet* UScriptableInteractiveTool::AddPropertySetOfType(
+	TSubclassOf<UScriptableInteractiveToolPropertySet> PropertySetType,
+	FString Identifier,
+	TEnumAsByte<EToolsFrameworkOutcomePins>& Outcome)
+{
+	Outcome = EToolsFrameworkOutcomePins::Failure;
+
+	if (NamedPropertySets.Contains(Identifier) || Identifier.Len() == 0 )
+	{
+		UE_LOG(LogTemp, Warning, TEXT("[AddPropertySetOfType] PropertySet Identifier %s is invalid or already in use!"), *Identifier);
+		return nullptr;
+	}
+
+	UClass* Class = PropertySetType.Get();
+	UClass* BasePropertySetClass = UScriptableInteractiveToolPropertySet::StaticClass();
+	if (Class->IsChildOf(BasePropertySetClass) == false)
+	{
+		UE_LOG(LogTemp, Warning, TEXT("[AddPropertySetOfType] Class Type %s is not a UScriptableInteractiveToolPropertySet Subclass!"), *Class->GetAuthoredName());
+		return nullptr;
+	}
+
+	UObject* NewPropertySetObj = NewObject<UScriptableInteractiveToolPropertySet>(this, Class);
+	if (NewPropertySetObj == nullptr)
+	{
+		UE_LOG(LogTemp, Warning, TEXT("[AddPropertySetOfType] failed to create instance of Class Type %s"), *Class->GetAuthoredName());
+		return nullptr;
+	}
+	UScriptableInteractiveToolPropertySet* NewPropertySet = Cast<UScriptableInteractiveToolPropertySet>(NewPropertySetObj);
+
+	// ensure that this property set is marked Transient, as it should never be saved
+	NewPropertySet->SetFlags(RF_Transient);
+
+	NewPropertySet->ParentTool = this;
+
+	AddToolPropertySource(NewPropertySet);
+	NamedPropertySets.Add(Identifier, NewPropertySet);
+
+	Outcome = EToolsFrameworkOutcomePins::Success;
+	return NewPropertySet;
+}
+
+
+void UScriptableInteractiveTool::RemovePropertySetByName(
+	FString Identifier,
+	TEnumAsByte<EToolsFrameworkOutcomePins>& Outcome)
+{
+	Outcome = EToolsFrameworkOutcomePins::Failure;
+
+	TWeakObjectPtr<UScriptableInteractiveToolPropertySet>* Found = NamedPropertySets.Find(Identifier);
+	if (Found != nullptr)
+	{
+		NamedPropertySets.Remove(Identifier);
+		if ((*Found).IsValid())
+		{
+			bool bRemoved = RemoveToolPropertySource( (*Found).Get() );
+			ensure(bRemoved == true);
+		}
+		Outcome = EToolsFrameworkOutcomePins::Success;
+	}
+}
+
+void UScriptableInteractiveTool::SetPropertySetVisibleByName(
+	FString Identifier,
+	bool bVisible)
+{
+	TWeakObjectPtr<UScriptableInteractiveToolPropertySet>* Found = NamedPropertySets.Find(Identifier);
+	if (Found != nullptr && (*Found).IsValid() )
+	{
+		SetToolPropertySourceEnabled( (*Found).Get(), bVisible );
+	}
+}
+
+void UScriptableInteractiveTool::ForcePropertySetUpdateByName(
+	FString Identifier)
+{
+	TWeakObjectPtr<UScriptableInteractiveToolPropertySet>* Found = NamedPropertySets.Find(Identifier);
+	if (Found != nullptr && (*Found).IsValid())
+	{
+		NotifyOfPropertyChangeByTool((*Found).Get());
+	}
+}
+
+
+void UScriptableInteractiveTool::RestorePropertySetSettings(
+	UScriptableInteractiveToolPropertySet* PropertySet,
+	FString SaveKey)
+{
+	if (PropertySet)
+	{
+		PropertySet->RestoreProperties(this, SaveKey);
+	}
+}
+
+
+void UScriptableInteractiveTool::SavePropertySetSettings(
+	UScriptableInteractiveToolPropertySet* PropertySet,
+	FString SaveKey)
+{
+	if (PropertySet)
+	{
+		PropertySet->SaveProperties(this, SaveKey);
+	}
+}
+
+namespace UELocal
+{
+
+template<typename PropertyType>
+PropertyType* FindValidPropertyByTypeAndName(
+	UScriptableInteractiveToolPropertySet* PropertySet,
+	FString PropertyName,
+	FString CallerString )
+{
+	if ( PropertySet == nullptr )
+	{
+		UE_LOG(LogTemp, Warning, TEXT("%s: Property Set is null"), *CallerString);
+		return nullptr;
+	}
+
+	for (TFieldIterator<FProperty> PropIt(PropertySet->GetClass(), EFieldIterationFlags::None); PropIt; ++PropIt)
+	{
+		FProperty* Property = *PropIt;
+		if (Property->GetName() == PropertyName)
+		{
+			if (PropertyType* TypedProp = CastField<PropertyType>(*PropIt))
+			{
+				return TypedProp;
+			}
+			else
+			{
+				UE_LOG(LogTemp, Warning, TEXT("%s: Property %s does not have a compatible Type"), *CallerString, *PropertyName);
+				return nullptr;
+			}
+		}
+	}
+
+	UE_LOG(LogTemp, Warning, TEXT("%s: Could not find a Property named %s in the Property Set %s"), *CallerString, *PropertyName, 
+		*PropertySet->GetName());
+	return nullptr;
+}
+
+} // end namespace UELocal
+
+
+void UScriptableInteractiveTool::WatchFloatProperty(
+	UScriptableInteractiveToolPropertySet* PropertySet,
+	FString PropertyName,
+	const FToolFloatPropertyModifiedDelegate& OnModified)
+{
+	if (FNumericProperty* NumericProp = UELocal::FindValidPropertyByTypeAndName<FNumericProperty>(
+			PropertySet, PropertyName, TEXT("WatchFloatProperty")))
+	{
+		if (NumericProp->IsFloatingPoint() == false)
+		{
+			UE_LOG(LogTemp, Warning, TEXT("WatchFloatProperty: Property %s is not Float type"), *PropertyName);
+			return;
+		}
+
+		int32 ArrayIndex = WatchFloatPropertyDelegates.Num();
+		WatchFloatPropertyDelegates.Add(OnModified);
+
+		PropertySet->WatchProperty<float>(
+			[NumericProp, PropertySet]()->float { 
+				// is this value stable? could it be captured instead?
+				const void* ValuePtr = NumericProp->ContainerPtrToValuePtr<void>(PropertySet);
+				return (double)NumericProp->GetFloatingPointPropertyValue(ValuePtr);
+			},
+			[this, PropertySet, PropertyName, ArrayIndex](double NewValue) { 
+				WatchFloatPropertyDelegates[ArrayIndex].ExecuteIfBound(PropertySet, PropertyName, NewValue);
+			} );
+	}
+
+}
+
+
+
+void UScriptableInteractiveTool::WatchIntProperty(
+	UScriptableInteractiveToolPropertySet* PropertySet,
+	FString PropertyName,
+	const FToolIntPropertyModifiedDelegate& OnModified)
+{
+	if (FNumericProperty* NumericProp = UELocal::FindValidPropertyByTypeAndName<FNumericProperty>(
+			PropertySet, PropertyName, TEXT("WatchIntProperty")))
+	{
+		if (NumericProp->IsInteger() == false)
+		{
+			UE_LOG(LogTemp, Warning, TEXT("WatchIntProperty: Property %s is not Integer type"), *PropertyName);
+			return;
+		}
+
+		int32 ArrayIndex = WatchIntPropertyDelegates.Num();
+		WatchIntPropertyDelegates.Add(OnModified);
+
+		PropertySet->WatchProperty<int>(
+			[NumericProp, PropertySet]()->int { 
+				const void* ValuePtr = NumericProp->ContainerPtrToValuePtr<void>(PropertySet);
+				return (int)NumericProp->GetSignedIntPropertyValue(ValuePtr);
+			},
+			[this, PropertySet, PropertyName, ArrayIndex](int NewValue) { 
+				WatchIntPropertyDelegates[ArrayIndex].ExecuteIfBound(PropertySet, PropertyName, NewValue);
+			} );
+	}
+}
+
+
+
+void UScriptableInteractiveTool::WatchBoolProperty(
+	UScriptableInteractiveToolPropertySet* PropertySet,
+	FString PropertyName,
+	const FToolBoolPropertyModifiedDelegate& OnModified)
+{
+	if (FBoolProperty* BoolProp = UELocal::FindValidPropertyByTypeAndName<FBoolProperty>(
+			PropertySet, PropertyName, TEXT("WatchBoolProperty")))
+	{
+		int32 ArrayIndex = WatchBoolPropertyDelegates.Num();
+		WatchBoolPropertyDelegates.Add(OnModified);
+
+		PropertySet->WatchProperty<bool>(
+			[BoolProp, PropertySet]()->bool { 
+				const void* ValuePtr = BoolProp->ContainerPtrToValuePtr<void>(PropertySet);
+				bool bValue = BoolProp->GetPropertyValue(ValuePtr);
+				return bValue;
+			},
+			[this, PropertySet, PropertyName, ArrayIndex](bool bNewValue) { 
+				WatchBoolPropertyDelegates[ArrayIndex].ExecuteIfBound(PropertySet, PropertyName, bNewValue);
+			} );
+	}
+}
+
+
+void UScriptableInteractiveTool::WatchEnumProperty(
+	UScriptableInteractiveToolPropertySet* PropertySet,
+	FString PropertyName,
+	const FToolEnumPropertyModifiedDelegate& OnModified)
+{
+	if (FEnumProperty* EnumProperty = UELocal::FindValidPropertyByTypeAndName<FEnumProperty>(
+			PropertySet, PropertyName, TEXT("WatchEnumProperty")))
+	{
+		FNumericProperty* EnumIntProperty = EnumProperty->GetUnderlyingProperty();
+		int32 ArrayIndex = WatchEnumPropertyDelegates.Num();
+		WatchEnumPropertyDelegates.Add(OnModified);
+
+		PropertySet->WatchProperty<int>(
+			[EnumIntProperty, EnumProperty, PropertySet]()->int64 { 
+				const void* ValuePtr = EnumProperty->ContainerPtrToValuePtr<void>(PropertySet);
+				int64 Value = EnumIntProperty->GetSignedIntPropertyValue(ValuePtr);
+				return Value;
+			},
+			[this, PropertySet, PropertyName, ArrayIndex](int64 NewValue) { 
+				WatchEnumPropertyDelegates[ArrayIndex].ExecuteIfBound(PropertySet, PropertyName, (uint8)NewValue);
+			} );
+	}
+}
+
+
+
+void UScriptableInteractiveTool::WatchStringProperty(
+	UScriptableInteractiveToolPropertySet* PropertySet,
+	FString PropertyName,
+	const FToolStringPropertyModifiedDelegate& OnModified)
+{
+	if (FStrProperty* StringProp = UELocal::FindValidPropertyByTypeAndName<FStrProperty>(
+			PropertySet, PropertyName, TEXT("WatchStringProperty")))
+	{
+		int32 ArrayIndex = WatchStringPropertyDelegates.Num();
+		WatchStringPropertyDelegates.Add(OnModified);
+
+		PropertySet->WatchProperty<FString>(
+			[StringProp, PropertySet]()->FString { 
+				const void* ValuePtr = StringProp->ContainerPtrToValuePtr<void>(PropertySet);
+				return StringProp->GetPropertyValue(ValuePtr);
+			},
+			[this, PropertySet, PropertyName, ArrayIndex](FString NewValue) { 
+				WatchStringPropertyDelegates[ArrayIndex].ExecuteIfBound(PropertySet, PropertyName, NewValue);
+			} );
+	}
+}
+
+
+void UScriptableInteractiveTool::WatchNameProperty(
+	UScriptableInteractiveToolPropertySet* PropertySet,
+	FString PropertyName,
+	const FToolFNamePropertyModifiedDelegate& OnModified)
+{
+	if (FNameProperty* NameProp = UELocal::FindValidPropertyByTypeAndName<FNameProperty>(
+			PropertySet, PropertyName, TEXT("WatchNameProperty")))
+	{
+		int32 ArrayIndex = WatchFNamePropertyDelegates.Num();
+		WatchFNamePropertyDelegates.Add(OnModified);
+
+		PropertySet->WatchProperty<FName>(
+			[NameProp, PropertySet]()->FName { 
+				const void* ValuePtr = NameProp->ContainerPtrToValuePtr<void>(PropertySet);
+				return NameProp->GetPropertyValue(ValuePtr);
+			},
+			[this, PropertySet, PropertyName, ArrayIndex](FName NewValue) { 
+				WatchFNamePropertyDelegates[ArrayIndex].ExecuteIfBound(PropertySet, PropertyName, NewValue);
+			} );
+	}
+}
+
+
+void UScriptableInteractiveTool::WatchObjectProperty(
+	UScriptableInteractiveToolPropertySet* PropertySet,
+	FString PropertyName,
+	const FToolObjectPropertyModifiedDelegate& OnModified)
+{
+	if (FObjectPropertyBase* ObjectProp = UELocal::FindValidPropertyByTypeAndName<FObjectPropertyBase>(
+			PropertySet, PropertyName, TEXT("WatchObjectProperty")))
+	{
+		int32 ArrayIndex = WatchObjectPropertyDelegates.Num();
+		WatchObjectPropertyDelegates.Add(OnModified);
+
+		PropertySet->WatchProperty<UObject*>(
+			[ObjectProp, PropertySet]() { 
+				const void* ValuePtr = ObjectProp->ContainerPtrToValuePtr<void>(PropertySet);
+				return ObjectProp->GetObjectPropertyValue(ValuePtr);
+			},
+			[this, PropertySet, PropertyName, ArrayIndex](UObject* NewValue) { 
+				WatchObjectPropertyDelegates[ArrayIndex].ExecuteIfBound(PropertySet, PropertyName, NewValue);
+			} );
+	}
+}
+
+
+
+
+void UScriptableInteractiveTool::WatchProperty(
+	UScriptableInteractiveToolPropertySet* PropertySet,
+	FString PropertyName,
+	const FToolPropertyModifiedDelegate& OnModified)
+{
+	FProperty* Property = UELocal::FindValidPropertyByTypeAndName<FProperty>(
+			PropertySet, PropertyName, TEXT("WatchProperty"));
+	if (Property == nullptr)
+	{
+		return;
+	}
+
+	int32 ArrayIndex = WatchAnyPropertyInfo.Num();
+	FAnyPropertyWatchInfo NewWatchInfo;
+	NewWatchInfo.Delegate = OnModified;
+	WatchAnyPropertyInfo.Add(NewWatchInfo);
+
+	// FMapProperty, FSetProperty, FDelegateProperty, FInterfaceProperty
+
+	if (FNumericProperty* NumericProp = CastField<FNumericProperty>(Property))
+	{
+		if (NumericProp->IsInteger())
+		{
+			WatchAnyPropertyInfo[ArrayIndex].KnownType = EAnyPropertyWatchTypes::Integer;
+			PropertySet->WatchProperty<int>(
+				[NumericProp, PropertySet]()->int { 
+					const void* ValuePtr = NumericProp->ContainerPtrToValuePtr<void>(PropertySet);
+					return (int)NumericProp->GetSignedIntPropertyValue(ValuePtr);
+				},
+				[this, PropertySet, PropertyName, ArrayIndex](int NewValue) { 
+					WatchAnyPropertyInfo[ArrayIndex].Delegate.ExecuteIfBound(PropertySet, PropertyName);
+				} );
+		} 
+		else if (NumericProp->IsFloatingPoint())
+		{
+			WatchAnyPropertyInfo[ArrayIndex].KnownType = EAnyPropertyWatchTypes::Double;
+			PropertySet->WatchProperty<float>(
+				[NumericProp, PropertySet]()->float { 
+					// is this value stable? could it be captured instead?
+					const void* ValuePtr = NumericProp->ContainerPtrToValuePtr<void>(PropertySet);
+					return (double)NumericProp->GetFloatingPointPropertyValue(ValuePtr);
+				},
+				[this, PropertySet, PropertyName, ArrayIndex](double NewValue) { 
+					WatchAnyPropertyInfo[ArrayIndex].Delegate.ExecuteIfBound(PropertySet, PropertyName);
+				} );
+		}
+		else
+		{
+			UE_LOG(LogTemp, Warning, TEXT("WatchProperty: Property %s is an unsupported type"), *PropertyName);
+		}
+	}
+	else if (FEnumProperty* EnumProperty = CastField<FEnumProperty>(Property))
+	{
+		WatchAnyPropertyInfo[ArrayIndex].KnownType = EAnyPropertyWatchTypes::Enum;
+		FNumericProperty* EnumIntProperty = EnumProperty->GetUnderlyingProperty();
+		PropertySet->WatchProperty<int>(
+			[EnumIntProperty, EnumProperty, PropertySet]()->int64 { 
+				const void* ValuePtr = EnumProperty->ContainerPtrToValuePtr<void>(PropertySet);
+				return (int)EnumIntProperty->GetSignedIntPropertyValue(ValuePtr);
+			},
+			[this, PropertySet, PropertyName, ArrayIndex](int64 NewValue) { 
+				WatchAnyPropertyInfo[ArrayIndex].Delegate.ExecuteIfBound(PropertySet, PropertyName);
+			} );
+	}
+	else if (FBoolProperty* BoolProp = CastField<FBoolProperty>(Property))
+	{
+		WatchAnyPropertyInfo[ArrayIndex].KnownType = EAnyPropertyWatchTypes::Bool;
+		PropertySet->WatchProperty<bool>(
+			[BoolProp, PropertySet]()->bool { 
+				const void* ValuePtr = BoolProp->ContainerPtrToValuePtr<void>(PropertySet);
+				return BoolProp->GetPropertyValue(ValuePtr);
+			},
+			[this, PropertySet, PropertyName, ArrayIndex](bool bNewValue) { 
+				WatchAnyPropertyInfo[ArrayIndex].Delegate.ExecuteIfBound(PropertySet, PropertyName);
+			} );
+	}
+	else if (FStrProperty* StringProperty = CastField<FStrProperty>(Property))
+	{
+		WatchAnyPropertyInfo[ArrayIndex].KnownType = EAnyPropertyWatchTypes::String;
+		PropertySet->WatchProperty<FString>(
+			[StringProperty, PropertySet]()->FString { 
+				const void* ValuePtr = StringProperty->ContainerPtrToValuePtr<void>(PropertySet);
+				return StringProperty->GetPropertyValue(ValuePtr);
+			},
+			[this, PropertySet, PropertyName, ArrayIndex](FString NewValue) { 
+				WatchAnyPropertyInfo[ArrayIndex].Delegate.ExecuteIfBound(PropertySet, PropertyName);
+			} );				
+	}
+	else if (FNameProperty* NameProp = CastField<FNameProperty>(Property))
+	{
+		WatchAnyPropertyInfo[ArrayIndex].KnownType = EAnyPropertyWatchTypes::FName;
+		PropertySet->WatchProperty<FName>(
+			[NameProp, PropertySet]()->FName { 
+				const void* ValuePtr = NameProp->ContainerPtrToValuePtr<void>(PropertySet);
+				return NameProp->GetPropertyValue(ValuePtr);
+			},
+			[this, PropertySet, PropertyName, ArrayIndex](FName NewValue) { 
+				WatchAnyPropertyInfo[ArrayIndex].Delegate.ExecuteIfBound(PropertySet, PropertyName);
+			} );
+	}
+	else if (FObjectPropertyBase* BaseObjProperty = CastField<FObjectPropertyBase>(Property) )
+	{
+		WatchAnyPropertyInfo[ArrayIndex].KnownType = EAnyPropertyWatchTypes::Object;
+		PropertySet->WatchProperty<UObject*>(
+			[BaseObjProperty, PropertySet]() { 
+				const void* ValuePtr = BaseObjProperty->ContainerPtrToValuePtr<void>(PropertySet);
+				return BaseObjProperty->GetObjectPropertyValue(ValuePtr);
+			},
+			[this, PropertySet, PropertyName, ArrayIndex](UObject* NewValue) { 
+				WatchAnyPropertyInfo[ArrayIndex].Delegate.ExecuteIfBound(PropertySet, PropertyName);
+			} );
+	}
+	else if (FStructProperty* StructProperty = CastField<FStructProperty>(Property))
+	{
+		WatchAnyPropertyInfo[ArrayIndex].KnownType = EAnyPropertyWatchTypes::Struct;
+		PropertySet->WatchProperty<uint32>(
+			[this, ArrayIndex, StructProperty, PropertySet]()->uint32 { 
+				const void* StructValuePtr = StructProperty->ContainerPtrToValuePtr<void>(PropertySet);
+				uint32 ValueHash = StructProperty->GetValueTypeHash(StructValuePtr);
+				return ValueHash;
+			},
+			[this, PropertySet, PropertyName, ArrayIndex](uint32 NewValue) { 
+				WatchAnyPropertyInfo[ArrayIndex].Delegate.ExecuteIfBound(PropertySet, PropertyName);
+			} );
+	}
+	else if (FArrayProperty* ArrayProperty = CastField<FArrayProperty>(Property))
+	{
+		WatchAnyPropertyInfo[ArrayIndex].KnownType = EAnyPropertyWatchTypes::Array;
+		PropertySet->WatchProperty<uint32>(
+			[this, ArrayIndex, ArrayProperty, PropertySet]()->uint32 { 
+				const void* ArrayValuePtr = ArrayProperty->ContainerPtrToValuePtr<void>(PropertySet);
+				FScriptArrayHelper ArrayHelper(ArrayProperty, ArrayValuePtr);
+				int32 ElementSize = ArrayProperty->Inner->ElementSize;
+				int32 NumElements = ArrayHelper.Num();
+ 				uint32 CRCValue = NumElements;
+				for ( int32 k = 0; k < NumElements; ++k )
+				{
+					CRCValue = FCrc::MemCrc32(ArrayHelper.GetRawPtr(k), ElementSize, CRCValue);
+				}
+				return CRCValue;
+			},
+			[this, PropertySet, PropertyName, ArrayIndex](uint32 NewValue) { 
+				WatchAnyPropertyInfo[ArrayIndex].Delegate.ExecuteIfBound(PropertySet, PropertyName);
+			} );
+	}
+	else
+	{
+		UE_LOG(LogTemp, Warning, TEXT("WatchProperty: Property Type for Property %s is not currently supported"), *PropertyName)
+
+		//PropertySet->WatchProperty<uint32>(
+		//	[this, ArrayIndex, Property, PropertySet]()->uint32 { 
+		//		TArray<uint8>& UseTempBuffer = WatchAnyPropertyInfo[ArrayIndex].TempBuffer;
+		//		if (UseTempBuffer.Num() != Property->GetSize())
+		//		{
+		//			UseTempBuffer.SetNum(Property->GetSize(), false);
+		//		}
+		//		const void* ValuePtr = Property->ContainerPtrToValuePtr<void>(PropertySet);
+		//		Property->CopyCompleteValue(UseTempBuffer.GetData(), ValuePtr);
+		//		uint32 CRCValue = FCrc::MemCrc32(UseTempBuffer.GetData(), UseTempBuffer.Num());
+		//		return CRCValue;
+		//	},
+		//	[this, PropertySet, PropertyName, ArrayIndex](uint32 NewValue) { 
+		//		WatchAnyPropertyInfo[ArrayIndex].Delegate.ExecuteIfBound(PropertySet, PropertyName);
+		//	} );
+	}
+
+}
+
+
+
+
+void UScriptableInteractiveTool::CreateTRSGizmo(
+	FString Identifier,
+	FTransform InitialTransform,
+	FScriptableToolGizmoOptions GizmoOptions,
+	TEnumAsByte<EToolsFrameworkOutcomePins>& Outcome)
+{
+	Outcome = EToolsFrameworkOutcomePins::Failure;
+
+	if (Gizmos.Contains(Identifier) || Identifier.Len() == 0)
+	{
+		UE_LOG(LogTemp, Warning, TEXT("[CreateTRSGizmo] Gizmo Identifier %s is invalid or already in use!"), *Identifier);
+		return;
+	}
+
+	ETransformGizmoSubElements SubElements = ETransformGizmoSubElements::None;
+	auto ConfigureGizmoTranslation = [&SubElements, &GizmoOptions](EScriptableToolGizmoTranslation SettingFlag, ETransformGizmoSubElements GizmoFlag)
+	{
+		EnumAddFlags(SubElements, ((GizmoOptions.TranslationParts & (int32)SettingFlag) != 0) ? GizmoFlag : ETransformGizmoSubElements::None);
+	};
+	ConfigureGizmoTranslation(EScriptableToolGizmoTranslation::TranslateAxisX, ETransformGizmoSubElements::TranslateAxisX);
+	ConfigureGizmoTranslation(EScriptableToolGizmoTranslation::TranslateAxisY, ETransformGizmoSubElements::TranslateAxisY);
+	ConfigureGizmoTranslation(EScriptableToolGizmoTranslation::TranslateAxisZ, ETransformGizmoSubElements::TranslateAxisZ);
+	ConfigureGizmoTranslation(EScriptableToolGizmoTranslation::TranslatePlaneXY, ETransformGizmoSubElements::TranslatePlaneXY);
+	ConfigureGizmoTranslation(EScriptableToolGizmoTranslation::TranslatePlaneXZ, ETransformGizmoSubElements::TranslatePlaneXZ);
+	ConfigureGizmoTranslation(EScriptableToolGizmoTranslation::TranslatePlaneYZ, ETransformGizmoSubElements::TranslatePlaneYZ);
+
+	auto ConfigureGizmoRotation = [&SubElements, &GizmoOptions](EScriptableToolGizmoRotation SettingFlag, ETransformGizmoSubElements GizmoFlag)
+	{
+		EnumAddFlags(SubElements, ((GizmoOptions.RotationParts & (int32)SettingFlag) != 0) ? GizmoFlag : ETransformGizmoSubElements::None);
+	};
+	ConfigureGizmoRotation(EScriptableToolGizmoRotation::RotateAxisX, ETransformGizmoSubElements::RotateAxisX);
+	ConfigureGizmoRotation(EScriptableToolGizmoRotation::RotateAxisY, ETransformGizmoSubElements::RotateAxisY);
+	ConfigureGizmoRotation(EScriptableToolGizmoRotation::RotateAxisZ, ETransformGizmoSubElements::RotateAxisZ);
+
+	auto ConfigureGizmoScale = [&SubElements, &GizmoOptions](EScriptableToolGizmoScale SettingFlag, ETransformGizmoSubElements GizmoFlag)
+	{
+		EnumAddFlags(SubElements, ((GizmoOptions.ScaleParts & (int32)SettingFlag) != 0) ? GizmoFlag : ETransformGizmoSubElements::None);
+	};
+	ConfigureGizmoScale(EScriptableToolGizmoScale::ScaleAxisX, ETransformGizmoSubElements::ScaleAxisX);
+	ConfigureGizmoScale(EScriptableToolGizmoScale::ScaleAxisY, ETransformGizmoSubElements::ScaleAxisY);
+	ConfigureGizmoScale(EScriptableToolGizmoScale::ScaleAxisZ, ETransformGizmoSubElements::ScaleAxisZ);
+	ConfigureGizmoScale(EScriptableToolGizmoScale::ScalePlaneXY, ETransformGizmoSubElements::ScalePlaneXY);
+	ConfigureGizmoScale(EScriptableToolGizmoScale::ScalePlaneXZ, ETransformGizmoSubElements::ScalePlaneXZ);
+	ConfigureGizmoScale(EScriptableToolGizmoScale::ScalePlaneYZ, ETransformGizmoSubElements::ScalePlaneYZ);
+
+	if (SubElements == ETransformGizmoSubElements::None)
+	{
+		UE_LOG(LogTemp, Warning, TEXT("[CreateTRSGizmo] Cannot create a Transform Gizmo thath as no sub-widgets"), *Identifier);
+		return;
+	}
+
+	UCombinedTransformGizmo* NewGizmo = ( GizmoOptions.bRepositionable ) ?
+		UE::TransformGizmoUtil::CreateCustomRepositionableTransformGizmo(GetToolManager(), SubElements, this, Identifier) :
+		UE::TransformGizmoUtil::CreateCustomTransformGizmo(GetToolManager(), SubElements, this, Identifier);
+	if (NewGizmo == nullptr)
+	{
+		UE_LOG(LogTemp, Warning, TEXT("[CreateTRSGizmo] Gizmo creation failed"), *Identifier);
+		return;
+	}
+
+	NewGizmo->bUseContextGizmoMode = (GizmoOptions.GizmoMode == EScriptableToolGizmoMode::FromViewportSettings);
+	switch (GizmoOptions.GizmoMode)
+	{
+		case EScriptableToolGizmoMode::Combined:
+			NewGizmo->ActiveGizmoMode = EToolContextTransformGizmoMode::Combined;
+			break;
+		case EScriptableToolGizmoMode::TranslationOnly:
+			NewGizmo->ActiveGizmoMode = EToolContextTransformGizmoMode::Translation;
+			break;
+		case EScriptableToolGizmoMode::RotationOnly:
+			NewGizmo->ActiveGizmoMode = EToolContextTransformGizmoMode::Rotation;
+			break;
+		case EScriptableToolGizmoMode::ScaleOnly:
+			NewGizmo->ActiveGizmoMode = EToolContextTransformGizmoMode::Scale;
+			break;
+	}
+
+	NewGizmo->bUseContextCoordinateSystem = (GizmoOptions.CoordSystem == EScriptableToolGizmoCoordinateSystem::FromViewportSettings);
+	switch (GizmoOptions.CoordSystem)
+	{
+		case EScriptableToolGizmoCoordinateSystem::Local:
+			NewGizmo->CurrentCoordinateSystem = EToolContextCoordinateSystem::Local;
+			break;
+		case EScriptableToolGizmoCoordinateSystem::World:
+			NewGizmo->CurrentCoordinateSystem = EToolContextCoordinateSystem::World;
+			break;
+	}
+
+	NewGizmo->bSnapToWorldGrid = GizmoOptions.bSnapTranslation;
+	NewGizmo->bSnapToWorldRotGrid = GizmoOptions.bSnapRotation;
+
+
+	UTransformProxy* GizmoProxy = NewObject<UTransformProxy>(this);
+	GizmoProxy->SetTransform(InitialTransform);
+	NewGizmo->SetActiveTarget(GizmoProxy, GetToolManager());		// do we always want to route through ToolManager here?
+
+	GizmoProxy->OnTransformChanged.AddWeakLambda(this, [this, Identifier](UTransformProxy*, FTransform NewTransform) 
+	{
+		OnGizmoTransformChanged_Handler(Identifier, NewTransform);
+	});
+
+	if (GizmoOptions.bAllowNegativeScaling == false)
+	{
+		NewGizmo->SetDisallowNegativeScaling(true);
+	}
+
+	Gizmos.Add(Identifier, NewGizmo);
+
+
+	Outcome = EToolsFrameworkOutcomePins::Success;
+}
+
+
+void UScriptableInteractiveTool::DestroyTRSGizmo(
+	FString Identifier,
+	TEnumAsByte<EToolsFrameworkOutcomePins>& Outcome)
+{
+	Outcome = EToolsFrameworkOutcomePins::Failure;
+
+	TObjectPtr<UCombinedTransformGizmo>* Found = Gizmos.Find(Identifier);
+	if ( Found == nullptr )
+	{
+		UE_LOG(LogTemp, Warning, TEXT("[DestroyTRSGizmo] Gizmo Identifier %s could not be found"), *Identifier);
+		return;
+	}
+
+	GetToolManager()->GetPairedGizmoManager()->DestroyGizmo(*Found);
+	Gizmos.Remove(Identifier);
+	Outcome = EToolsFrameworkOutcomePins::Success;
+}
+
+
+
+void UScriptableInteractiveTool::SetGizmoVisible(
+	FString Identifier,
+	bool bVisible)
+{
+	TObjectPtr<UCombinedTransformGizmo>* Found = Gizmos.Find(Identifier);
+	if (Found != nullptr)
+	{
+		(*Found)->SetVisibility(bVisible);
+	}
+	else
+	{
+		UE_LOG(LogTemp, Warning, TEXT("[SetGizmoVisible] Gizmo Identifier %s could not be found"), *Identifier);
+	}
+}
+
+
+
+void UScriptableInteractiveTool::SetGizmoTransform(
+	FString Identifier,
+	FTransform NewTransform,
+	bool bUndoable)
+{
+	TObjectPtr<UCombinedTransformGizmo>* Found = Gizmos.Find(Identifier);
+	if (Found != nullptr)
+	{
+		if (bUndoable)
+		{
+			(*Found)->SetNewGizmoTransform(NewTransform);
+		}
+		else
+		{
+			(*Found)->ReinitializeGizmoTransform(NewTransform);
+		}
+	}
+	else
+	{
+		UE_LOG(LogTemp, Warning, TEXT("[SetGizmoTransform] Gizmo Identifier %s could not be found"), *Identifier);
+	}
+}
+
+
+FTransform UScriptableInteractiveTool::GetGizmoTransform(
+	FString Identifier)
+{
+	TObjectPtr<UCombinedTransformGizmo>* Found = Gizmos.Find(Identifier);
+	if (Found != nullptr)
+	{
+		return (*Found)->GetGizmoTransform();
+	}
+	else
+	{
+		UE_LOG(LogTemp, Warning, TEXT("[GetGizmoTransform] Gizmo Identifier %s could not be found"), *Identifier);
+		return FTransform::Identity;
+	}
+}
+
+
+
+void UScriptableInteractiveTool::OnGizmoTransformChanged_Handler(FString GizmoIdentifier, FTransform NewTransform)
+{
+	// forward to BlueprintImplementableEvent
+	OnGizmoTransformChanged(GizmoIdentifier, NewTransform);
+}
+
+
+
+
+
+void UScriptableInteractiveTool::AddLogMessage(FText Message, bool bHighlighted)
+{
+	if (UInteractiveToolManager* ToolManager = GetToolManager())
+	{
+		if (bHighlighted)
+		{
+			UE_LOG(LogTemp, Warning, TEXT("%s"), *Message.ToString() );
+		}
+		else
+		{
+			UE_LOG(LogTemp, Log, TEXT("%s"), *Message.ToString() );
+		}
+	}
+}
+
+void UScriptableInteractiveTool::DisplayUserHelpMessage(FText Message)
+{
+	if (UInteractiveToolManager* ToolManager = GetToolManager())
+	{
+		ToolManager->DisplayMessage(Message, EToolMessageLevel::UserNotification);
+	}
+}
+
+void UScriptableInteractiveTool::DisplayUserWarningMessage(FText Message)
+{
+	if (UInteractiveToolManager* ToolManager = GetToolManager())
+	{
+		ToolManager->DisplayMessage(Message, EToolMessageLevel::UserWarning);
+	}
+}
+
+void UScriptableInteractiveTool::ClearUserMessages(bool bNotifications, bool bWarnings)
+{
+	if (UInteractiveToolManager* ToolManager = GetToolManager())
+	{
+		if (bNotifications)
+		{
+			ToolManager->DisplayMessage(FText(), EToolMessageLevel::UserNotification);
+		}
+		if (bWarnings)
+		{
+			ToolManager->DisplayMessage(FText(), EToolMessageLevel::UserWarning);
+		}
+	}
+}
+
+
+
+
+
+
+FInputRayHit UScriptableToolsUtilityLibrary::MakeInputRayHit_Miss()
+{
+	return FInputRayHit();
+}
+
+FInputRayHit UScriptableToolsUtilityLibrary::MakeInputRayHit_MaxDepth()
+{
+	return FInputRayHit(TNumericLimits<float>::Max());
+}
+
+FInputRayHit UScriptableToolsUtilityLibrary::MakeInputRayHit(double HitDepth, UObject* OptionalHitObject)
+{
+	FInputRayHit Result(HitDepth);
+	Result.HitObject = OptionalHitObject;
+	return Result;
+}
+
+#undef LOCTEXT_NAMESPACE
\ No newline at end of file
diff --git a/Engine/Plugins/Experimental/ScriptableToolsFramework/Source/ScriptableToolsFramework/Private/ScriptableToolBuilder.cpp b/Engine/Plugins/Experimental/ScriptableToolsFramework/Source/ScriptableToolsFramework/Private/ScriptableToolBuilder.cpp
new file mode 100644
--- /dev/null
+++ b/Engine/Plugins/Experimental/ScriptableToolsFramework/Source/ScriptableToolsFramework/Private/ScriptableToolBuilder.cpp
@@ -0,0 +1,25 @@
+// Copyright Epic Games, Inc. All Rights Reserved.
+
+#include "ScriptableToolBuilder.h"
+#include "InteractiveToolManager.h"
+#include "UObject/ObjectMacros.h"
+
+
+bool UBaseScriptableToolBuilder::CanBuildTool(const FToolBuilderState& SceneState) const
+{
+	return ToolClass.IsValid();
+}
+
+
+UInteractiveTool* UBaseScriptableToolBuilder::BuildTool(const FToolBuilderState& SceneState) const
+{
+	UClass* UseClass = ToolClass.Get();
+
+	UObject* NewToolObj = NewObject<UScriptableInteractiveTool>(SceneState.ToolManager, UseClass);
+	check(NewToolObj != nullptr);
+	UScriptableInteractiveTool* NewTool = Cast<UScriptableInteractiveTool>(NewToolObj);
+	NewTool->SetTargetWorld(SceneState.World);
+	check(NewTool != nullptr);
+	return NewTool;
+}
+
diff --git a/Engine/Plugins/Experimental/ScriptableToolsFramework/Source/ScriptableToolsFramework/Private/ScriptableToolSet.cpp b/Engine/Plugins/Experimental/ScriptableToolsFramework/Source/ScriptableToolsFramework/Private/ScriptableToolSet.cpp
new file mode 100644
--- /dev/null
+++ b/Engine/Plugins/Experimental/ScriptableToolsFramework/Source/ScriptableToolsFramework/Private/ScriptableToolSet.cpp
@@ -0,0 +1,109 @@
+// Copyright Epic Games, Inc. All Rights Reserved.
+
+#include "ScriptableToolSet.h"
+
+#include "ScriptableInteractiveTool.h"
+#include "ScriptableToolBuilder.h"
+
+#include "UObject/Class.h"
+#include "UObject/ObjectMacros.h"
+#include "UObject/UObjectGlobals.h"
+#include "UObject/SoftObjectPath.h"
+#include "Templates/SubclassOf.h"
+#include "Kismet2/KismetEditorUtilities.h"
+
+#include "Modules/ModuleManager.h"
+#include "AssetRegistry/AssetRegistryModule.h"
+#include "AssetRegistry/IAssetRegistry.h"
+
+
+void UScriptableToolSet::ReinitializeScriptableTools()
+{
+	UClass* ScriptableToolClass = UScriptableInteractiveTool::StaticClass();
+	UScriptableInteractiveTool* ScriptableToolCDO = ScriptableToolClass->GetDefaultObject<UScriptableInteractiveTool>();
+
+	TSet<UClass*> PotentialToolClasses;
+
+	// Iterate over Blueprint classes to try to find UScriptableInteractiveTool blueprints
+	// Note that this code may not be fully reliable, but it appears to work so far...
+
+	// Load the asset registry module
+	FAssetRegistryModule& AssetRegistryModule = FModuleManager::LoadModuleChecked<FAssetRegistryModule>(FName("AssetRegistry"));
+	IAssetRegistry& AssetRegistry = AssetRegistryModule.Get();
+	// The asset registry is populated asynchronously at startup, so there's no guarantee it has finished.
+	// This simple approach just runs a synchronous scan on the entire content directory.
+	// Better solutions would be to specify only the path to where the relevant blueprints are,
+	// or to register a callback with the asset registry to be notified of when it's finished populating.
+	TArray< FString > ContentPaths;
+	ContentPaths.Add(TEXT("/Game"));
+	AssetRegistry.ScanPathsSynchronous(ContentPaths);
+
+	FTopLevelAssetPath BaseClassPath = ScriptableToolClass->GetClassPathName();
+
+	// Use the asset registry to get the set of all class names deriving from Base
+	TArray<FTopLevelAssetPath> BaseModeClasses;
+	BaseModeClasses.Add(BaseClassPath);
+	TSet<FTopLevelAssetPath> DerivedClassPaths;
+	AssetRegistry.GetDerivedClassNames(BaseModeClasses, TSet<FTopLevelAssetPath>(), DerivedClassPaths);
+
+	for (const FTopLevelAssetPath& ClassPath : DerivedClassPaths)
+	{
+		// don't include Base tools  (note this will also catch EditorScriptableToolsFramework)
+		if (ClassPath.ToString().Contains(TEXT("ScriptableToolsFramework")))
+		{
+			continue;
+		}
+
+		FSoftObjectPath ObjectPath(ClassPath.ToString());
+		TSoftClassPtr<UScriptableInteractiveTool> SoftClass = TSoftClassPtr<UScriptableInteractiveTool>(ObjectPath);
+		if (UClass* Class = SoftClass.LoadSynchronous())
+		{
+			if (Class->HasAnyClassFlags(CLASS_Abstract) || (Class->GetAuthoritativeClass() != Class))
+			{
+				continue;
+			}
+
+			PotentialToolClasses.Add(Class);
+		}
+	}
+
+	Tools.Reset();
+	ToolBuilders.Reset();
+
+	// if the class is viable, create a ToolBuilder for it
+	for (UClass* Class : PotentialToolClasses)
+	{
+		if (Class->IsChildOf(ScriptableToolClass))
+		{
+			FScriptableToolInfo ToolInfo;
+			ToolInfo.ToolClass = Class;
+			ToolInfo.ToolCDO = Class->GetDefaultObject<UScriptableInteractiveTool>();
+
+			UBaseScriptableToolBuilder* ToolBuilder = NewObject<UBaseScriptableToolBuilder>(this);
+			ToolBuilder->ToolClass = Class;
+			ToolBuilders.Add(ToolBuilder);
+
+			ToolInfo.ToolBuilder = ToolBuilder;
+			Tools.Add(ToolInfo);
+		}
+	}
+
+
+}
+
+
+
+
+void UScriptableToolSet::ForEachScriptableTool(
+	TFunctionRef<void(UClass* ToolClass, UBaseScriptableToolBuilder* ToolBuilder)> ProcessToolFunc)
+{
+	for (FScriptableToolInfo& ToolInfo : Tools)
+	{
+		if (ToolInfo.ToolClass.IsValid() && ToolInfo.ToolBuilder.IsValid())
+		{
+			ProcessToolFunc(ToolInfo.ToolClass.Get(), ToolInfo.ToolBuilder.Get());
+		}
+	}
+}
+
+
diff --git a/Engine/Plugins/Experimental/ScriptableToolsFramework/Source/ScriptableToolsFramework/Private/ScriptableToolsFrameworkModule.cpp b/Engine/Plugins/Experimental/ScriptableToolsFramework/Source/ScriptableToolsFramework/Private/ScriptableToolsFrameworkModule.cpp
new file mode 100644
--- /dev/null
+++ b/Engine/Plugins/Experimental/ScriptableToolsFramework/Source/ScriptableToolsFramework/Private/ScriptableToolsFrameworkModule.cpp
@@ -0,0 +1,20 @@
+// Copyright Epic Games, Inc. All Rights Reserved.
+
+#include "ScriptableToolsFrameworkModule.h"
+
+#define LOCTEXT_NAMESPACE "FScriptableToolsFrameworkModule"
+
+void FScriptableToolsFrameworkModule::StartupModule()
+{
+	// This code will execute after your module is loaded into memory; the exact timing is specified in the .uplugin file per-module
+}
+
+void FScriptableToolsFrameworkModule::ShutdownModule()
+{
+	// This function may be called during shutdown to clean up your module.  For modules that support dynamic reloading,
+	// we call this function before unloading the module.
+}
+
+#undef LOCTEXT_NAMESPACE
+	
+IMPLEMENT_MODULE(FScriptableToolsFrameworkModule, ScriptableToolsFramework)
\ No newline at end of file
diff --git a/Engine/Plugins/Experimental/ScriptableToolsFramework/Source/ScriptableToolsFramework/Public/BaseTools/ScriptableClickDragTool.h b/Engine/Plugins/Experimental/ScriptableToolsFramework/Source/ScriptableToolsFramework/Public/BaseTools/ScriptableClickDragTool.h
new file mode 100644
--- /dev/null
+++ b/Engine/Plugins/Experimental/ScriptableToolsFramework/Source/ScriptableToolsFramework/Public/BaseTools/ScriptableClickDragTool.h
@@ -0,0 +1,190 @@
+// Copyright Epic Games, Inc. All Rights Reserved.
+
+#pragma once
+
+#include "ScriptableInteractiveTool.h"
+#include "InputState.h"
+#include "BaseBehaviors/BehaviorTargetInterfaces.h"
+
+#include "ScriptableClickDragTool.generated.h"
+
+class UClickDragInputBehavior;
+class UMouseHoverBehavior;
+
+/**
+ * UScriptableClickDragTool is an extension of UScriptableInteractiveTool that supports
+ * a "click-drag" interaction, ie capture the input device on press-down and collect any
+ * device movement events until release. Hover is also supported. 
+ * 
+ * Note: Currently these interactions are only supported for mouse devices, or devices remapped to 
+ * be a mouse by the UInputRouter. Also currently only the left-mouse-button is supported.
+ */
+UCLASS()
+class SCRIPTABLETOOLSFRAMEWORK_API UScriptableClickDragTool : public UScriptableInteractiveTool, public IClickDragBehaviorTarget, public IHoverBehaviorTarget
+{
+	GENERATED_BODY()
+
+public:
+	/**
+	 * Enable Hover support API functions OnHoverHitTest / OnHoverBegin / OnHoverUpdate / OnHoverEnd for Mouse devices.
+	 * Defaults to disabled. Cannot be changed dynamically after Setup runs.
+	 */
+	UPROPERTY(BlueprintReadWrite, EditAnywhere, Category = "SingleClick Tool Settings")
+	bool bWantMouseHover = false;
+
+	/**
+	 * Allow Modifiers states to be updated during ClickDrag sequence. Defaults to enabled.
+	 * Cannot be changed dynamically after Setup runs.
+	 */
+	UPROPERTY(BlueprintReadWrite, EditAnywhere, Category = "SingleClick Tool Settings")
+	bool bUpdateModifiersDuringDrag = true;
+
+	//
+	// Click Support API
+	//
+public:
+
+	UFUNCTION(BlueprintNativeEvent, Category = "ClickDrag Events")
+	FInputRayHit TestIfCanBeginClickDrag(FInputDeviceRay ClickPos, const FScriptableToolModifierStates& Modifiers);
+
+	virtual FInputRayHit TestIfCanBeginClickDrag_Implementation(FInputDeviceRay ClickPos, const FScriptableToolModifierStates& Modifiers);
+
+
+	UFUNCTION(BlueprintNativeEvent, Category = "ClickDrag Events")
+	void OnDragBegin(FInputDeviceRay StartPosition, const FScriptableToolModifierStates& Modifiers);
+
+	virtual void OnDragBegin_Implementation(FInputDeviceRay StartPosition, const FScriptableToolModifierStates& Modifiers);
+
+
+	UFUNCTION(BlueprintNativeEvent, Category = "ClickDrag Events")
+	void OnDragUpdatePosition(FInputDeviceRay NewPosition, const FScriptableToolModifierStates& Modifiers);
+
+	virtual void OnDragUpdatePosition_Implementation(FInputDeviceRay NewPosition, const FScriptableToolModifierStates& Modifiers);
+
+
+	UFUNCTION(BlueprintNativeEvent, Category = "ClickDrag Events")
+	void OnDragEnd(FInputDeviceRay EndPosition, const FScriptableToolModifierStates& Modifiers);
+
+	virtual void OnDragEnd_Implementation(FInputDeviceRay EndPosition, const FScriptableToolModifierStates& Modifiers);
+
+
+	UFUNCTION(BlueprintNativeEvent, Category = "ClickDrag Events")
+	void OnDragSequenceCancelled();
+
+	virtual void OnDragSequenceCancelled_Implementation();
+
+	UFUNCTION(BlueprintCallable, Category = "ScriptableTool|Input")
+	UPARAM(DisplayName="In ClickDrag State") bool InActiveClickDrag() const;
+
+	//
+	// Hover Support API
+	//
+public:
+
+	/**
+	 * OnHoverHitTest is called to allow the Tool to indicate if it would like to start consuing "hover" input at HoverPos.
+	 * The Tool can return yes/no and a "hit depth", which will be used to determine if the Tool is given the active hover input stream.
+	 * The default OnHoverHitTest implement always captures the hover.
+	 * 
+	 * Once the hover is accepted, the OnHoverBegin event will fire, and then a stream of OnHoverUpdate events. If at any point
+	 * the hover is no longer relevant, OnHoverUpdate should return false. OnHoverEnd will be called at that point, or if 
+	 * the hover is cancelled for any reason (eg mouse goes out of window, button is pressed, etc)
+	 * 
+	 * @param HoverPos the current position of the cursor/device, including both a 3D ray from the eye, and (optionally) a 2D mouse position
+	 * @param Modifiers current modifier key/button state
+	 * @return a FInputRayHit indicating a hit (true/false) and the hit depth along the hit-ray (0 for "always consume")
+	 */
+	UFUNCTION(BlueprintNativeEvent, Category = "Hover Events")
+	FInputRayHit OnHoverHitTest(FInputDeviceRay HoverPos, const FScriptableToolModifierStates& Modifiers);
+
+	virtual FInputRayHit OnHoverHitTest_Implementation(FInputDeviceRay HoverPos, const FScriptableToolModifierStates& Modifiers);
+
+	/**
+	 * Begin a sequence of Hover input events 
+	 * @param HoverPos the current position of the cursor/device, including both a 3D ray from the eye, and (optionally) a 2D mouse position
+	 * @param Modifiers current modifier key/button state
+	 */
+	UFUNCTION(BlueprintNativeEvent, Category = "Hover Events")
+	void OnHoverBegin(FInputDeviceRay HoverPos, const FScriptableToolModifierStates& Modifiers);
+
+	virtual void OnHoverBegin_Implementation(FInputDeviceRay HoverPos, const FScriptableToolModifierStates& Modifiers);
+
+	/**
+	 * Update an active Hover sequence. OnHoverUpdate is only ever called between OnHoverBegin and OnHoverEnd
+	 * @param HoverPos the current position of the cursor/device, including both a 3D ray from the eye, and (optionally) a 2D mouse position
+	 * @param Modifiers current modifier key/button state
+	 * @return true to continue hovering, false to stop receiving additional hover events
+	 */
+	UFUNCTION(BlueprintNativeEvent, Category = "Hover Events")
+	UPARAM(DisplayName="Continue Hovering") bool
+	OnHoverUpdate(FInputDeviceRay HoverPos, const FScriptableToolModifierStates& Modifiers);
+
+	virtual bool OnHoverUpdate_Implementation(FInputDeviceRay HoverPos, const FScriptableToolModifierStates& Modifiers);
+
+	/**
+	 * Indicate that an active captured Hover sequence has ended. The device may no longer be in the viewport, so no position can be provided.
+	 * @param Modifiers current modifier key/button state
+	 */
+	UFUNCTION(BlueprintNativeEvent, Category = "Hover Events")
+	void OnHoverEnd(const FScriptableToolModifierStates& Modifiers);
+
+	virtual void OnHoverEnd_Implementation(const FScriptableToolModifierStates& Modifiers);
+
+	UFUNCTION(BlueprintCallable, Category = "ScriptableTool|Input")
+	UPARAM(DisplayName="In Hover State") bool InActiveHover() const;
+
+	//
+	// Modifer Buttons Support API
+	//
+public:
+
+	/** @return true if the Shift key is currently held down */
+	UFUNCTION(BlueprintCallable, Category="ScriptableTool|Input")
+	UPARAM(DisplayName="Shift Down") bool IsShiftDown() const;
+
+	/** @return true if the Ctrl key is currently held down */
+	UFUNCTION(BlueprintCallable, Category="ScriptableTool|Input")
+	UPARAM(DisplayName="Ctrl Down") bool IsCtrlDown() const;
+
+	/** @return true if the Alt key is currently held down */
+	UFUNCTION(BlueprintCallable, Category="ScriptableTool|Input")
+	UPARAM(DisplayName="Alt Down") bool IsAltDown() const;
+
+	/** @return a struct containing the current Modifier key states */
+	UFUNCTION(BlueprintCallable, Category="ScriptableTool|Input")
+	FScriptableToolModifierStates GetActiveModifiers();
+
+
+
+public:
+	virtual void Setup() override;
+
+	UPROPERTY(Transient, DuplicateTransient, NonTransactional, SkipSerialization)
+	TObjectPtr<UClickDragInputBehavior> ClickDragBehavior;
+
+	bool bInClickDrag = false;
+
+	UPROPERTY(Transient, DuplicateTransient, NonTransactional, SkipSerialization)
+	TObjectPtr<UMouseHoverBehavior> MouseHoverBehavior;
+
+	bool bInHover = false;
+
+	bool bShiftModifier = false;
+	bool bCtrlModifier = false;
+	bool bAltModifier = false;
+
+	// IClickDragBehaviorTarget API
+	virtual FInputRayHit CanBeginClickDragSequence(const FInputDeviceRay& PressPos) override;
+	virtual void OnClickPress(const FInputDeviceRay& PressPos) override;
+	virtual void OnClickDrag(const FInputDeviceRay& DragPos) override;
+	virtual void OnClickRelease(const FInputDeviceRay& ReleasePos) override;
+	virtual void OnTerminateDragSequence() override;
+	virtual void OnUpdateModifierState(int ModifierID, bool bIsOn);
+
+	// IHoverBehaviorTarget API
+	virtual FInputRayHit BeginHoverSequenceHitTest(const FInputDeviceRay& PressPos);
+	virtual void OnBeginHover(const FInputDeviceRay& DevicePos);
+	virtual bool OnUpdateHover(const FInputDeviceRay& DevicePos);
+	virtual void OnEndHover();
+
+};
\ No newline at end of file
diff --git a/Engine/Plugins/Experimental/ScriptableToolsFramework/Source/ScriptableToolsFramework/Public/BaseTools/ScriptableSingleClickTool.h b/Engine/Plugins/Experimental/ScriptableToolsFramework/Source/ScriptableToolsFramework/Public/BaseTools/ScriptableSingleClickTool.h
new file mode 100644
--- /dev/null
+++ b/Engine/Plugins/Experimental/ScriptableToolsFramework/Source/ScriptableToolsFramework/Public/BaseTools/ScriptableSingleClickTool.h
@@ -0,0 +1,178 @@
+// Copyright Epic Games, Inc. All Rights Reserved.
+
+#pragma once
+
+#include "ScriptableInteractiveTool.h"
+#include "InputState.h"
+#include "BaseBehaviors/BehaviorTargetInterfaces.h"
+
+#include "ScriptableSingleClickTool.generated.h"
+
+class USingleClickInputBehavior;
+class UMouseHoverBehavior;
+
+/**
+ * UScriptableSingleClickTool is an extension of UScriptableInteractiveTool that supports
+ * a "click" interaction, ie button press-down-and-release on an arbitrary Target (defined by the Tool).
+ * Note that like a standard GUI click interaction, if the device is moved off the "Target" 
+ * before the button is released, the click is cancelled. Note that in this context, the 
+ * Target is defined implicitly by the Tool returning a "hit" for a given input device location,
+ * rather than any specific geometry.
+ * 
+ * Hover is also supported. 
+ * 
+ * Note: Currently these interactions are only supported for mouse devices, or devices remapped to 
+ * be a mouse by the UInputRouter. Also currently only the left-mouse-button is supported.
+ */
+UCLASS()
+class SCRIPTABLETOOLSFRAMEWORK_API UScriptableSingleClickTool : public UScriptableInteractiveTool, public IClickBehaviorTarget, public IHoverBehaviorTarget
+{
+	GENERATED_BODY()
+
+public:
+	/**
+	 * Enable Hover support API functions OnHoverHitTest / OnHoverBegin / OnHoverUpdate / OnHoverEnd for Mouse devices.
+	 * Defaults to disabled.
+	 */
+	UPROPERTY(BlueprintReadWrite, EditAnywhere, Category = "SingleClick Tool Settings")
+	bool bWantMouseHover = false;
+
+	//
+	// Click Support API
+	//
+public:
+	/**
+	 * TestIfHitByClick is called to allow the Tool to indicate if it would like to consume a potential click at ClickPos.
+	 * The Tool can return yes/no and a "hit depth", which will be used to determine if the Tool is given the click (ie if it has the nearest depth),
+	 * at which point the OnHitByClick event will fire. The default TestIfHitByClick implementation always captures the click (at depth 0).
+	 * 
+	 * Note that this function will be called twice for a particular hit - once on mouse-down and once on mouse-up, and if it returns no-hit on mouse-up, 
+	 * the click will be "cancelled", similar to standard GUI click interactions.
+	 * 
+	 * @param ClickPos the position of the click, including both a 3D ray from the eye, and (optionally) a 2D mouse position
+	 * @param Modifiers current modifier key/button state
+	 * @return a FInputRayHit indicating a hit (true/false) and the hit depth along the hit-ray (0 for "always consume")
+	 */
+	UFUNCTION(BlueprintNativeEvent, Category = "Click Events")
+	FInputRayHit TestIfHitByClick(FInputDeviceRay ClickPos, const FScriptableToolModifierStates& Modifiers);
+
+	virtual FInputRayHit TestIfHitByClick_Implementation(FInputDeviceRay ClickPos, const FScriptableToolModifierStates& Modifiers);
+
+	/**
+	 * OnHitByClick is called when the Tool has indicated it would like to consume a click event (via TestIfHitByClick), and the click was not cancelled
+	 * or consumed by a nearer object. 
+	 * @param ClickPos the position of the click, including both a 3D ray from the eye, and (optionally) a 2D mouse position
+	 * @param Modifiers current modifier key/button state
+	 */
+	UFUNCTION(BlueprintNativeEvent, Category = "Click Events")
+	void OnHitByClick(FInputDeviceRay ClickPos, const FScriptableToolModifierStates& Modifiers);
+
+	virtual void OnHitByClick_Implementation(FInputDeviceRay ClickPos, const FScriptableToolModifierStates& Modifiers);
+
+
+	//
+	// Hover Support API
+	//
+public:
+
+	/**
+	 * OnHoverHitTest is called to allow the Tool to indicate if it would like to start consuing "hover" input at HoverPos.
+	 * The Tool can return yes/no and a "hit depth", which will be used to determine if the Tool is given the active hover input stream.
+	 * The default OnHoverHitTest implement always captures the hover.
+	 * 
+	 * Once the hover is accepted, the OnHoverBegin event will fire, and then a stream of OnHoverUpdate events. If at any point
+	 * the hover is no longer relevant, OnHoverUpdate should return false. OnHoverEnd will be called at that point, or if 
+	 * the hover is cancelled for any reason (eg mouse goes out of window, button is pressed, etc)
+	 * 
+	 * @param HoverPos the current position of the cursor/device, including both a 3D ray from the eye, and (optionally) a 2D mouse position
+	 * @param Modifiers current modifier key/button state
+	 * @return a FInputRayHit indicating a hit (true/false) and the hit depth along the hit-ray (0 for "always consume")
+	 */
+	UFUNCTION(BlueprintNativeEvent, Category = "Hover Events")
+	FInputRayHit OnHoverHitTest(FInputDeviceRay HoverPos, const FScriptableToolModifierStates& Modifiers);
+
+	virtual FInputRayHit OnHoverHitTest_Implementation(FInputDeviceRay HoverPos, const FScriptableToolModifierStates& Modifiers);
+
+	/**
+	 * Begin a sequence of Hover input events 
+	 * @param HoverPos the current position of the cursor/device, including both a 3D ray from the eye, and (optionally) a 2D mouse position
+	 * @param Modifiers current modifier key/button state
+	 */
+	UFUNCTION(BlueprintNativeEvent, Category = "Hover Events")
+	void OnHoverBegin(FInputDeviceRay HoverPos, const FScriptableToolModifierStates& Modifiers);
+
+	virtual void OnHoverBegin_Implementation(FInputDeviceRay HoverPos, const FScriptableToolModifierStates& Modifiers);
+
+	/**
+	 * Update an active Hover sequence. OnHoverUpdate is only ever called between OnHoverBegin and OnHoverEnd
+	 * @param HoverPos the current position of the cursor/device, including both a 3D ray from the eye, and (optionally) a 2D mouse position
+	 * @param Modifiers current modifier key/button state
+	 * @return true to continue hovering, false to stop receiving additional hover events
+	 */
+	UFUNCTION(BlueprintNativeEvent, Category = "Hover Events")
+	UPARAM(DisplayName="Continue Hovering") bool
+	OnHoverUpdate(FInputDeviceRay HoverPos, const FScriptableToolModifierStates& Modifiers);
+
+	virtual bool OnHoverUpdate_Implementation(FInputDeviceRay HoverPos, const FScriptableToolModifierStates& Modifiers);
+
+	/**
+	 * Indicate that an active captured Hover sequence has ended. The device may no longer be in the viewport, so no position can be provided.
+	 * @param Modifiers current modifier key/button state
+	 */
+	UFUNCTION(BlueprintNativeEvent, Category = "Hover Events")
+	void OnHoverEnd(const FScriptableToolModifierStates& Modifiers);
+
+	virtual void OnHoverEnd_Implementation(const FScriptableToolModifierStates& Modifiers);
+
+	UFUNCTION(BlueprintCallable, Category = "ScriptableTool|Input")
+	UPARAM(DisplayName="In Hover State") bool InActiveHover() const;
+
+	//
+	// Modifer Buttons Support API
+	//
+public:
+
+	/** @return true if the Shift key is currently held down */
+	UFUNCTION(BlueprintCallable, Category="ScriptableTool|Input")
+	UPARAM(DisplayName="Shift Down") bool IsShiftDown() const;
+
+	/** @return true if the Ctrl key is currently held down */
+	UFUNCTION(BlueprintCallable, Category="ScriptableTool|Input")
+	UPARAM(DisplayName="Ctrl Down") bool IsCtrlDown() const;
+
+	/** @return true if the Alt key is currently held down */
+	UFUNCTION(BlueprintCallable, Category="ScriptableTool|Input")
+	UPARAM(DisplayName="Alt Down") bool IsAltDown() const;
+
+	/** @return a struct containing the current Modifier key states */
+	UFUNCTION(BlueprintCallable, Category="ScriptableTool|Input")
+	FScriptableToolModifierStates GetActiveModifiers();
+
+
+
+public:
+	virtual void Setup() override;
+
+	UPROPERTY(Transient, DuplicateTransient, NonTransactional, SkipSerialization)
+	TObjectPtr<USingleClickInputBehavior> SingleClickBehavior;
+
+	UPROPERTY(Transient, DuplicateTransient, NonTransactional, SkipSerialization)
+	TObjectPtr<UMouseHoverBehavior> MouseHoverBehavior;
+
+	bool bInHover = false;
+
+	bool bShiftModifier = false;
+	bool bCtrlModifier = false;
+	bool bAltModifier = false;
+
+	// IClickBehaviorTarget API
+	virtual FInputRayHit IsHitByClick(const FInputDeviceRay& ClickPos);
+	virtual void OnClicked(const FInputDeviceRay& ClickPos);
+	virtual void OnUpdateModifierState(int ModifierID, bool bIsOn);
+
+	// IHoverBehaviorTarget API
+	virtual FInputRayHit BeginHoverSequenceHitTest(const FInputDeviceRay& PressPos);
+	virtual void OnBeginHover(const FInputDeviceRay& DevicePos);
+	virtual bool OnUpdateHover(const FInputDeviceRay& DevicePos);
+	virtual void OnEndHover();
+};
\ No newline at end of file
diff --git a/Engine/Plugins/Experimental/ScriptableToolsFramework/Source/ScriptableToolsFramework/Public/ScriptableInteractiveTool.h b/Engine/Plugins/Experimental/ScriptableToolsFramework/Source/ScriptableToolsFramework/Public/ScriptableInteractiveTool.h
new file mode 100644
--- /dev/null
+++ b/Engine/Plugins/Experimental/ScriptableToolsFramework/Source/ScriptableToolsFramework/Public/ScriptableInteractiveTool.h
@@ -0,0 +1,740 @@
+// Copyright Epic Games, Inc. All Rights Reserved.
+
+#pragma once
+
+#include "InteractiveTool.h"
+#include "InteractiveGizmo.h"
+#include "Kismet/BlueprintFunctionLibrary.h"
+#include "ScriptableInteractiveTool.generated.h"
+
+class UWorld;
+class UCombinedTransformGizmo;
+class UTransformProxy;
+class UScriptableInteractiveTool;
+class FToolDataVisualizer;
+class FCanvas;
+class FSceneView;
+
+UENUM(BlueprintType)
+enum EToolsFrameworkOutcomePins
+{
+	Success,
+	Failure
+};
+
+
+UENUM(BlueprintType)
+enum class EScriptableToolShutdownType : uint8
+{
+	Complete = 0,
+	AcceptCancel = 1
+};
+
+USTRUCT(BlueprintType)
+struct FScriptableToolModifierStates
+{
+	GENERATED_BODY()
+
+	UPROPERTY(EditAnywhere, BlueprintReadWrite, Category = Modifiers)
+	bool bShiftDown = false;
+
+	UPROPERTY(EditAnywhere, BlueprintReadWrite, Category = Modifiers)
+	bool bCtrlDown = false;
+
+	UPROPERTY(EditAnywhere, BlueprintReadWrite, Category = Modifiers)
+	bool bAltDown = false;
+};
+
+
+
+UENUM(BlueprintType)
+enum class EScriptableToolGizmoMode : uint8
+{
+	TranslationOnly = 0,
+	RotationOnly = 1,
+	ScaleOnly = 2,
+	Combined = 3,
+	FromViewportSettings = 4
+};
+
+UENUM(BlueprintType)
+enum class EScriptableToolGizmoCoordinateSystem : uint8
+{
+	World = 0,
+	Local = 1,
+	FromViewportSettings = 2
+};
+
+UENUM(BlueprintType, meta = (Bitflags, UseEnumValuesAsMaskValuesInEditor = "true"))
+enum class EScriptableToolGizmoTranslation : uint8
+{
+	None = 0,
+
+	TranslateAxisX = 1 << 1,
+	TranslateAxisY = 1 << 2,
+	TranslateAxisZ = 1 << 3,
+	TranslatePlaneXY = 1 << 4,
+	TranslatePlaneXZ = 1 << 5,
+	TranslatePlaneYZ = 1 << 6,
+
+	All = 0xFF UMETA(Hidden)
+};
+ENUM_CLASS_FLAGS(EScriptableToolGizmoTranslation);
+
+
+UENUM(BlueprintType, meta = (Bitflags, UseEnumValuesAsMaskValuesInEditor = "true"))
+enum class EScriptableToolGizmoRotation : uint8
+{
+	None = 0,
+
+	RotateAxisX = 1 << 1,
+	RotateAxisY = 1 << 2,
+	RotateAxisZ = 1 << 3,
+
+	All = 0xFF UMETA(Hidden)
+};
+ENUM_CLASS_FLAGS(EScriptableToolGizmoRotation);
+
+
+UENUM(BlueprintType, meta = (Bitflags, UseEnumValuesAsMaskValuesInEditor = "true"))
+enum class EScriptableToolGizmoScale : uint8
+{
+	None = 0,
+
+	ScaleAxisX = 1 << 1,
+	ScaleAxisY = 1 << 2,
+	ScaleAxisZ = 1 << 3,
+	ScalePlaneXY = 1 << 4,
+	ScalePlaneXZ = 1 << 5,
+	ScalePlaneYZ = 1 << 6,
+	ScaleUniform = 1 << 7,
+
+	All = 0xFF UMETA(Hidden)
+};
+ENUM_CLASS_FLAGS(EScriptableToolGizmoScale);
+
+
+/**
+ * FScriptableToolGizmoOptions is a configuration struct passed to the CreateTRSGizmo function
+ * of UScriptableInteractiveTool. 
+ */
+USTRUCT(BlueprintType)
+struct FScriptableToolGizmoOptions
+{
+	GENERATED_BODY()
+
+	UPROPERTY(EditAnywhere, BlueprintReadWrite, Category = Gizmo)
+	EScriptableToolGizmoMode GizmoMode = EScriptableToolGizmoMode::Combined;
+
+	UPROPERTY(EditAnywhere, BlueprintReadWrite, Category = Gizmo, meta = (DisplayName = "Translation Elements", Bitmask, BitmaskEnum = "/Script/ScriptableToolsFramework.EScriptableToolGizmoTranslation",
+	ValidEnumValues = "TranslateAxisX, TranslateAxisY, TranslateAxisZ, TranslatePlaneXY, TranslatePlaneXZ, TranslatePlaneYZ"))
+	int32 TranslationParts = static_cast<int32>(EScriptableToolGizmoTranslation::All);
+
+	UPROPERTY(EditAnywhere, BlueprintReadWrite, Category = Gizmo, meta = (DisplayName = "Rotation Elements", Bitmask, BitmaskEnum = "/Script/ScriptableToolsFramework.EScriptableToolGizmoRotation",
+	ValidEnumValues = "RotateAxisX, RotateAxisY, RotateAxisZ"))
+	int32 RotationParts = static_cast<int32>(EScriptableToolGizmoRotation::All);
+
+	UPROPERTY(EditAnywhere, BlueprintReadWrite, Category = Gizmo, meta = (DisplayName = "Scale Elements", Bitmask, BitmaskEnum = "/Script/ScriptableToolsFramework.EScriptableToolGizmoScale",
+	ValidEnumValues = "ScaleAxisX, ScaleAxisY, ScaleAxisZ, ScalePlaneXY, ScalePlaneXZ, ScalePlaneYZ"))
+	int32 ScaleParts = static_cast<int32>(EScriptableToolGizmoScale::All);
+
+	UPROPERTY(EditAnywhere, BlueprintReadWrite, Category = Gizmo)
+	EScriptableToolGizmoCoordinateSystem CoordSystem = EScriptableToolGizmoCoordinateSystem::Local;
+
+	UPROPERTY(EditAnywhere, BlueprintReadWrite, Category = Gizmo)
+	bool bSnapTranslation = true;
+
+	UPROPERTY(EditAnywhere, BlueprintReadWrite, Category = Gizmo)
+	bool bSnapRotation = true;
+
+	UPROPERTY(EditAnywhere, BlueprintReadWrite, Category = Gizmo)
+	bool bRepositionable = false;
+
+	UPROPERTY(EditAnywhere, BlueprintReadWrite, Category = Gizmo)
+	bool bAllowNegativeScaling = true;
+};
+
+
+
+/**
+ * UScriptableTool_RenderAPI is helper Object that is created internally by a UScriptableInteractiveTool
+ * to allow Blueprints to access basic 3D rendering functionality, in the context of a specific Tool.
+ * The OnScriptRender event is called with an instance of this type.
+ */
+UCLASS(Transient)
+class SCRIPTABLETOOLSFRAMEWORK_API UScriptableTool_RenderAPI : public UObject
+{
+	GENERATED_BODY()
+public:
+	UFUNCTION(BlueprintCallable, Category = "ScriptableTool|Render")
+	UPARAM(DisplayName="Render Object") UScriptableTool_RenderAPI* 
+	DrawLine(FVector Start, FVector End, FLinearColor Color, float Thickness = 1.0, float DepthBias = 0.0f, bool bDepthTested = true);
+
+	// DrawPoint, DrawCircle, DrawViewFacingCircle
+	// DrawWireCylinder, DrawWireBox, DrawSquare
+	// PushTransform, PopTransform, SetTransform, PopAllTransforms
+
+protected:
+	friend class UScriptableInteractiveTool;
+
+	FToolDataVisualizer* ActiveVisualizer = nullptr;
+};
+
+/**
+ * UScriptableTool_HUDAPI is helper Object that is created internally by a UScriptableInteractiveTool
+ * to allow Blueprints to access basic 2D rendering functionality, in the context of a specific Tool.
+ * The OnScriptDrawHUD event is called with an instance of this type.
+ */
+UCLASS(Transient)
+class SCRIPTABLETOOLSFRAMEWORK_API UScriptableTool_HUDAPI : public UObject
+{
+	GENERATED_BODY()
+public:
+	UFUNCTION(BlueprintCallable, Category = "ScriptableTool|DrawHUD", Meta = (Color = "(R=1.0,G=1.0,B=1.0,A=1.0)"))
+	UPARAM(DisplayName="HUD Object") UScriptableTool_HUDAPI* 
+	DrawTextAtLocation(FVector Location, FString String, FLinearColor Color, bool bCentered, float ShiftRowsY);
+
+	UFUNCTION(BlueprintCallable, Category = "ScriptableTool|DrawHUD", Meta = (Color = "(R=1.0,G=1.0,B=1.0,A=1.0)"))
+	UPARAM(DisplayName="HUD Object") UScriptableTool_HUDAPI* 
+	DrawTextArrayAtLocation(FVector Location, TArray<FString> Strings, FLinearColor Color, bool bCentered, float ShiftRowsY);
+
+	UFUNCTION(BlueprintCallable, Category = "ScriptableTool|DrawHUD")
+	UPARAM(DisplayName="HUD Object") UScriptableTool_HUDAPI* 
+	GetCanvasLocation(FVector Location, FVector2D& CanvasLocation );
+
+
+	// DrawItem to draw line/etc (FCanvasLineItem, FCanvasBoxItem, FCanvasTriangleItem, FCanvasNGonItem)
+	// DrawTile
+	// DrawMaterial? via FCanvasTileItem
+
+protected:
+	friend class UScriptableInteractiveTool;
+	FCanvas* ActiveCanvas = nullptr;
+	const FSceneView* ActiveSceneView = nullptr;
+};
+
+
+
+/**
+ * UScriptableInteractiveToolPropertySet is a Blueprintable extension of UInteractiveToolPropertySet.
+ * This is a helper type meant to be used with UScriptableInteractiveTool. The intention is that
+ * the "Settings" of a particular Tool are stored in UScriptableInteractiveToolPropertySet ("Property Set") instances.
+ * The function UScriptableInteractiveTool::AddPropertySetOfType can be used to create and attach
+ * a new Property Set. Then these settings can be automatically exposed by, for example, the Editor Mode UI, 
+ * similar to (eg) Modeling Mode in the UE Editor.
+ * 
+ * In addition, UScriptableInteractiveTool has a set of functions like WatchFloatProperty, WatchBoolProperty, etc,
+ * that can be used to detect and respond to changes in the Property Set. This works with both changes caused
+ * by Editor UI (eg a Details Panel) as well as done directly in Blueprints or even C++ code. 
+  * 
+ * Note, however, that this Property Set mechanism completely optional, a Tool builder is free to use any
+ * method whatsoever to store/modify Tool settings.
+ */
+UCLASS(Transient, Blueprintable)
+class SCRIPTABLETOOLSFRAMEWORK_API UScriptableInteractiveToolPropertySet : public UInteractiveToolPropertySet
+{
+	GENERATED_BODY()
+
+	/** Access the Tool that owns this PropertySet */
+	UFUNCTION(BlueprintCallable, Category = "ScriptableTool", Meta = (ExpandEnumAsExecs = "Outcome"))
+	UScriptableInteractiveTool* GetOwningTool(
+		TEnumAsByte<EToolsFrameworkOutcomePins>& Outcome);
+
+protected:
+	TWeakObjectPtr<UScriptableInteractiveTool> ParentTool;
+	friend class UScriptableInteractiveTool;
+};
+
+
+
+// these are delegates for the various property-watchers below
+DECLARE_DYNAMIC_DELEGATE_TwoParams(FToolPropertyModifiedDelegate, UScriptableInteractiveToolPropertySet*, PropertySet, FString, PropertyName);
+DECLARE_DYNAMIC_DELEGATE_ThreeParams(FToolFloatPropertyModifiedDelegate, UScriptableInteractiveToolPropertySet*, PropertySet, FString, PropertyName, double, NewValue);
+DECLARE_DYNAMIC_DELEGATE_ThreeParams(FToolIntPropertyModifiedDelegate, UScriptableInteractiveToolPropertySet*, PropertySet, FString, PropertyName, int, NewValue);
+DECLARE_DYNAMIC_DELEGATE_ThreeParams(FToolBoolPropertyModifiedDelegate, UScriptableInteractiveToolPropertySet*, PropertySet, FString, PropertyName, bool, bNewValue);
+DECLARE_DYNAMIC_DELEGATE_ThreeParams(FToolEnumPropertyModifiedDelegate, UScriptableInteractiveToolPropertySet*, PropertySet, FString, PropertyName, uint8, NewValue);
+DECLARE_DYNAMIC_DELEGATE_ThreeParams(FToolStringPropertyModifiedDelegate, UScriptableInteractiveToolPropertySet*, PropertySet, FString, PropertyName, FString, NewValue);
+DECLARE_DYNAMIC_DELEGATE_ThreeParams(FToolFNamePropertyModifiedDelegate, UScriptableInteractiveToolPropertySet*, PropertySet, FString, PropertyName, FName, NewValue);
+DECLARE_DYNAMIC_DELEGATE_ThreeParams(FToolObjectPropertyModifiedDelegate, UScriptableInteractiveToolPropertySet*, PropertySet, FString, PropertyName, UObject*, NewValue);
+
+
+/**
+ * UScriptableInteractiveTool is an extension of UInteractiveTool that allows the Tool functionality to be 
+ * defined via Blueprints. 
+ */
+UCLASS(Transient, Blueprintable)
+class SCRIPTABLETOOLSFRAMEWORK_API UScriptableInteractiveTool : public UInteractiveTool
+{
+	GENERATED_BODY()
+public:
+	/** Name of this Tool, will be used in (eg) Toolbars */
+	UPROPERTY(BlueprintReadWrite, EditAnywhere, Category = "Scriptable Tool Settings")
+	FText ToolName;
+
+	/** Long Name of this Tool, will be used in (eg) longer labels like the Accept/Cancel toolbar */
+	UPROPERTY(BlueprintReadWrite, EditAnywhere, Category = "Scriptable Tool Settings", meta=(DisplayName="Long Name"))
+	FText ToolLongName;
+
+	/** Category of this Tool, will be used in (eg) Tool Palette Section headers */
+	UPROPERTY(BlueprintReadWrite, EditAnywhere, Category = "Scriptable Tool Settings", meta=(DisplayName="Category"))
+	FText ToolCategory;
+
+	/** Tooltip used for this Tool in (eg) icons/etc */
+	UPROPERTY(BlueprintReadWrite, EditAnywhere, Category = "Scriptable Tool Settings", meta=(DisplayName="Tooltip"))
+	FText ToolTooltip;
+
+	/** Specifies how the user exits this Tool, either Accept/Cancel-style or Complete-style */
+	UPROPERTY(BlueprintReadWrite, EditAnywhere, Category = "Scriptable Tool Settings", meta=(DisplayName="Shutdown Type"))
+	EScriptableToolShutdownType ToolShutdownType = EScriptableToolShutdownType::Complete;
+
+
+	/**
+	 * Implement OnScriptSetup to do initial setup/configuration of the Tool, such as adding
+	 * Property Sets, creating Gizmos, etc
+	 */
+	UFUNCTION(BlueprintImplementableEvent, Category = "ScriptableTool|Events")
+	void OnScriptSetup();
+
+	virtual void Setup() override;
+
+
+	/**
+	 * OnScriptTick is called every Editor Tick, ie basically every frame.
+	 * Implement per-frame processing here.
+	 */
+	UFUNCTION(BlueprintImplementableEvent, Category = "ScriptableTool|Events")
+	void OnScriptTick(float DeltaTime);
+
+	virtual void OnTick(float DeltaTime) override;
+
+
+	/**
+	 * CanAccept function. Implement this function for AcceptCancel Tools, and return
+	 * true when it is valid for the Tool Accept button to be active. Defaults to always true.
+	 */
+	UFUNCTION(BlueprintNativeEvent, Category = "ScriptableTool|Events")
+	bool OnScriptCanAccept() const;
+
+	bool OnScriptCanAccept_Implementation() const;
+
+	virtual bool HasAccept() const override;
+	virtual bool HasCancel() const override;
+	virtual bool CanAccept() const override;
+
+
+
+	/**
+	 * OnScriptShutdown is called when the Tool is shutting down. The ShutdownType defines
+	 * what the Tool should do. For Complete-style Tools, there is no difference, but for
+	 * Accept/Cancel-style Tools, on Accept the Tool should "commit" whatever it is previewing/etc,
+	 * and on Cancel it should roll back / do-nothing. Which of these occurs is based on the
+	 * ToolShutdownType property.
+	 * 
+	 * Tool Shutdown cannot be interrupted/aborted
+	 */
+	UFUNCTION(BlueprintImplementableEvent, Category = "ScriptableTool|Events")
+	void OnScriptShutdown(EToolShutdownType ShutdownType);
+
+	virtual void Shutdown(EToolShutdownType ShutdownType) override;
+
+	/**
+	 * Request that the active Tool be shut down. The Tool can post this to shut itself down (eg in a Tool where
+	 * the interaction paradigm is "click to do something and exit"). 
+	 * @param bAccept if this is an Accept/Cancel Tool, Accept it, otherwise Cancel. Ignored for Complete-type Tools.
+	 */
+	UFUNCTION(BlueprintCallable, Category = "ScriptableTool", Meta = (Identifier="Settings"))
+	void RequestToolShutdown(
+		bool bAccept,
+		bool bShowUserPopupMessage,
+		FText UserMessage);
+
+
+
+protected:
+	virtual void PostInitProperties() override;
+
+
+protected:
+	UPROPERTY(Transient, DuplicateTransient, NonTransactional, SkipSerialization)
+	TWeakObjectPtr<UWorld> TargetWorld;
+
+	virtual UWorld* GetWorld() const override;		// must implement this for some reason, otherwise many BP functions will not show up
+
+public:
+	virtual void SetTargetWorld(UWorld* World);		// needs to be accessible to ToolBuilder
+
+	/**
+	 * Access the World this Tool is currently operating on.
+	 */
+	UFUNCTION(BlueprintCallable, Category = "ScriptableTool")
+	UWorld* GetToolWorld();
+
+
+	//
+	// Render and DrawHUD support
+	//
+public:
+	virtual void Render(IToolsContextRenderAPI* RenderAPI) override;
+
+	/**
+	 * OnScriptRender is called every frame. Use the RenderAPI object to draw
+	 * various simple geometric elements like lines and points. This drawing
+	 * is not very efficient but is useful for basic Tool visualization/feedback.
+	 */
+	UFUNCTION(BlueprintImplementableEvent, Category = "ScriptableTool|Events")
+	void OnScriptRender( UPARAM(DisplayName="Render Object") UScriptableTool_RenderAPI* RenderAPI);
+
+	virtual void DrawHUD(FCanvas* Canvas, IToolsContextRenderAPI* RenderAPI) override;
+
+	/**
+	 * OnScriptDrawHUD is called every frame. Use the DrawHUDAPI to draw various
+	 * simple HUD elements like text. This drawing is not very efficient but
+	 * is useful for basic Tool visualization/feedback.
+	 */
+	UFUNCTION(BlueprintImplementableEvent, Category = "ScriptableTool|Events")
+	void OnScriptDrawHUD( UPARAM(DisplayName="HUD Object") UScriptableTool_HUDAPI* DrawHUDAPI);
+
+
+protected:
+	UPROPERTY(Transient, DuplicateTransient, NonTransactional, SkipSerialization)
+	TObjectPtr<UScriptableTool_RenderAPI> RenderHelper;
+
+	UPROPERTY(Transient, DuplicateTransient, NonTransactional, SkipSerialization)
+	TObjectPtr<UScriptableTool_HUDAPI> DrawHUDHelper;
+
+
+	//
+	// Property Set support
+	//
+public:
+
+	/**
+	 * Create a new Tool Property Set (ie BP subclass of UScriptableInteractiveToolPropertySet) with the given string Identifier
+	 * and attach it to the Tool. The public variables of this Property Set object will appear in (eg) Mode Details Panels, etc.
+	 * Multiple Property Sets of the same Type can be attached to a Tool, but each must have a unique Identifier.
+	 */
+	UFUNCTION(BlueprintCallable, Category = "ScriptableTool|PropertySets", Meta = (ExpandEnumAsExecs = "Outcome", Identifier="Settings"))
+	UPARAM(DisplayName="New Property Set") UScriptableInteractiveToolPropertySet* 
+	AddPropertySetOfType(
+		TSubclassOf<UScriptableInteractiveToolPropertySet> PropertySetType,
+		FString Identifier,
+		TEnumAsByte<EToolsFrameworkOutcomePins>& Outcome);
+
+	/**
+	 * Remove a Property Set from the current Tool, found via it's unique Identifier.
+	 * Unless the Property Set absolutely needs to stop existing, it is likely preferable to use SetPropertySetVisibleByName() to simply hide it.
+	 */
+	UFUNCTION(BlueprintCallable, Category = "ScriptableTool|PropertySets", Meta = (ExpandEnumAsExecs = "Outcome", Identifier="Settings"))
+	void RemovePropertySetByName(
+		FString Identifier,
+		TEnumAsByte<EToolsFrameworkOutcomePins>& Outcome);
+
+	/**
+	 * Set the visibility of a Property Set that is paired with the given unique Identifier.
+	 */
+	UFUNCTION(BlueprintCallable, Category = "ScriptableTool|PropertySets", Meta = (Identifier="Settings"))
+	void SetPropertySetVisibleByName(
+		FString Identifier,
+		bool bVisible);
+
+	/**
+	 * Force the Property Set associated with the given Identifier to be updated. 
+	 * This may be necessary if you change the Property Set values directly via Blueprints, and want to
+	 * see your changes reflected in (eg) a Details Panel showing the Tool Property values. 
+	 * Currently these changes cannot be detected automatically, necessitating this function.
+	 */
+	UFUNCTION(BlueprintCallable, Category = "ScriptableTool|PropertySets", Meta = (Identifier="Settings"))
+	void ForcePropertySetUpdateByName(
+		FString Identifier);
+
+	/**
+	 * Restore the values of the specified PropertySet, optionally with a specific SaveKey string.
+	 * This will have no effect unless SavePropertySetSettings() has been called on a compatible Property Set
+	 * in the same Engine session (ie in a previous invocation of the Tool, or another Tool that uses the same Property Set)
+	 */
+	UFUNCTION(BlueprintCallable, Category = "ScriptableTool|PropertySets", Meta = (Identifier="Settings"))
+	void RestorePropertySetSettings(
+		UScriptableInteractiveToolPropertySet* PropertySet,
+		FString SaveKey);
+
+	/**
+	 * Save the values of the specified PropertySet, optionally with a specific SaveKey string.
+	 * These saved values can be restored in future Tool invocations based on the SaveKey.
+	 */
+	UFUNCTION(BlueprintCallable, Category = "ScriptableTool|PropertySets", Meta = (Identifier="Settings"))
+	void SavePropertySetSettings(
+		UScriptableInteractiveToolPropertySet* PropertySet,
+		FString SaveKey);
+
+
+
+	//
+	// Property Watcher support
+	//
+public:
+	/**
+	 * Watch a Float-valued Property for changes (double precision)
+	 * @param PropertySet the Property Set which contains the desired Property to watch.
+	 * @param PropertyName the string name of the Property in the given Property Set
+	 * @param OnModified this delegate will be called if the Property value changes.
+	 */
+	UFUNCTION(BlueprintCallable, Category = "ScriptableTool|PropertySets")
+	void WatchFloatProperty(
+		UScriptableInteractiveToolPropertySet* PropertySet,
+		FString PropertyName,
+		const FToolFloatPropertyModifiedDelegate& OnModified );
+
+	/**
+	 * Watch an Integer-valued Property for changes
+	 * @param PropertySet the Property Set which contains the desired Property to watch.
+	 * @param PropertyName the string name of the Property in the given Property Set
+	 * @param OnModified this delegate will be called if the Property value changes.
+	 */
+	UFUNCTION(BlueprintCallable, Category = "ScriptableTool|PropertySets")
+	void WatchIntProperty(
+		UScriptableInteractiveToolPropertySet* PropertySet,
+		FString PropertyName,
+		const FToolIntPropertyModifiedDelegate& OnModified );
+
+	/**
+	 * Watch a Bool-valued Property for changes
+	 * @param PropertySet the Property Set which contains the desired Property to watch.
+	 * @param PropertyName the string name of the Property in the given Property Set
+	 * @param OnModified this delegate will be called if the Property value changes.
+	 */
+	UFUNCTION(BlueprintCallable, Category = "ScriptableTool|PropertySets")
+	void WatchBoolProperty(
+		UScriptableInteractiveToolPropertySet* PropertySet,
+		FString PropertyName,
+		const FToolBoolPropertyModifiedDelegate& OnModified );
+
+	/**
+	 * Watch an Enum-valued Property for changes. Note that in this case the OnModified
+	 * delegate will be called with a uint8 integer, which can be cast back to the original Enum type.
+	 * @param PropertySet the Property Set which contains the desired Property to watch.
+	 * @param PropertyName the string name of the Property in the given Property Set
+	 * @param OnModified this delegate will be called if the Property value changes.
+	 */
+	UFUNCTION(BlueprintCallable, Category = "ScriptableTool|PropertySets")
+	void WatchEnumProperty(
+		UScriptableInteractiveToolPropertySet* PropertySet,
+		FString PropertyName,
+		const FToolEnumPropertyModifiedDelegate& OnModified );
+
+	/**
+	 * Watch a String-valued Property for changes
+	 * @param PropertySet the Property Set which contains the desired Property to watch.
+	 * @param PropertyName the string name of the Property in the given Property Set
+	 * @param OnModified this delegate will be called if the Property value changes.
+	 */
+	UFUNCTION(BlueprintCallable, Category = "ScriptableTool|PropertySets")
+	void WatchStringProperty(
+		UScriptableInteractiveToolPropertySet* PropertySet,
+		FString PropertyName,
+		const FToolStringPropertyModifiedDelegate& OnModified );
+
+	/**
+	 * Watch an (F)Name-valued Property for changes
+	 * @param PropertySet the Property Set which contains the desired Property to watch.
+	 * @param PropertyName the string name of the Property in the given Property Set
+	 * @param OnModified this delegate will be called if the Property value changes.
+	 */
+	UFUNCTION(BlueprintCallable, Category = "ScriptableTool|PropertySets")
+	void WatchNameProperty(
+		UScriptableInteractiveToolPropertySet* PropertySet,
+		FString PropertyName,
+		const FToolFNamePropertyModifiedDelegate& OnModified );
+
+	/**
+	 * Watch an Object-valued Property for changes (ie UObject, UClass, etc)
+	 * @param PropertySet the Property Set which contains the desired Property to watch.
+	 * @param PropertyName the string name of the Property in the given Property Set
+	 * @param OnModified this delegate will be called if the Property value changes.
+	 */
+	UFUNCTION(BlueprintCallable, Category = "ScriptableTool|PropertySets")
+	void WatchObjectProperty(
+		UScriptableInteractiveToolPropertySet* PropertySet,
+		FString PropertyName,
+		const FToolObjectPropertyModifiedDelegate& OnModified );
+
+	/**
+	 * Watch any Property in a PropertySet for changes. This function handles most basic
+	 * properties, as well as Struct and Array properties, at some cost. The OnModified
+	 * delegate cannot be passed the new property value, it must be fetched from the PropertySet.
+	 * 
+	 * Note also that in the case of Structs and Arrays, changes are currently detected by hashes,
+	 * and there is always a small chance of hash collision.
+	 * 
+	 * @param PropertySet the Property Set which contains the desired Property to watch.
+	 * @param PropertyName the string name of the Property in the given Property Set
+	 * @param OnModified this delegate will be called if the Property value changes.
+	 */
+	UFUNCTION(BlueprintCallable, Category = "ScriptableTool|PropertySets")
+	void WatchProperty(
+		UScriptableInteractiveToolPropertySet* PropertySet,
+		FString PropertyName,
+		const FToolPropertyModifiedDelegate& OnModified );
+
+protected:
+	TMap<FString, TWeakObjectPtr<UScriptableInteractiveToolPropertySet>> NamedPropertySets;
+
+	TArray<FToolFloatPropertyModifiedDelegate> WatchFloatPropertyDelegates;
+	TArray<FToolIntPropertyModifiedDelegate> WatchIntPropertyDelegates;
+	TArray<FToolBoolPropertyModifiedDelegate> WatchBoolPropertyDelegates;
+	TArray<FToolEnumPropertyModifiedDelegate> WatchEnumPropertyDelegates;
+	TArray<FToolStringPropertyModifiedDelegate> WatchStringPropertyDelegates;
+	TArray<FToolFNamePropertyModifiedDelegate> WatchFNamePropertyDelegates;
+	TArray<FToolObjectPropertyModifiedDelegate> WatchObjectPropertyDelegates;
+
+
+	enum class EAnyPropertyWatchTypes
+	{
+		Integer = 0,
+		Bool = 1,
+		Double = 2,
+		Enum = 3,
+		String = 4,
+		FName = 5,
+		Object = 6,
+
+		Struct = 20,
+		Array = 30,
+
+
+		Unknown = 100
+	};
+	struct FAnyPropertyWatchInfo
+	{
+		EAnyPropertyWatchTypes KnownType = EAnyPropertyWatchTypes::Unknown;
+		TArray<uint8> TempBuffer;
+		FToolPropertyModifiedDelegate Delegate;
+	};
+	TArray<FAnyPropertyWatchInfo> WatchAnyPropertyInfo;
+
+	// gizmo API
+	//    - option to register delegate for specific gizmo?
+	//    - some way to easily connect gizmo transform to a FTransform property?
+public:
+
+	/**
+	 * Create a Translate/Rotate/Scale Gizmo with the given Options at the specified InitialTransform.
+	 * The Gizmo must be given a unique Identifier, which will be used to access it in other functions.
+	 * The Gizmo can be explicitly destroyed with DestroyTRSGizmo(), or it will be
+	 * automatically destroyed when the Tool exits
+	 */
+	UFUNCTION(BlueprintCallable, Category = "ScriptableTool|Gizmos", Meta = (ExpandEnumAsExecs = "Outcome", Identifier="Gizmo1"))
+	void CreateTRSGizmo(
+		FString Identifier,
+		FTransform InitialTransform,
+		FScriptableToolGizmoOptions GizmoOptions,
+		TEnumAsByte<EToolsFrameworkOutcomePins>& Outcome);
+
+	/**
+	 * Destroy a created Gizmo by name Identifier
+	 */
+	UFUNCTION(BlueprintCallable, Category = "ScriptableTool|Gizmos", Meta = (ExpandEnumAsExecs = "Outcome", Identifier="Gizmo1"))
+	void DestroyTRSGizmo(
+		FString Identifier,
+		TEnumAsByte<EToolsFrameworkOutcomePins>& Outcome);
+
+	/**
+	 * Set an existing Gizmo visible/hidden based on its name Identifier
+	 */
+	UFUNCTION(BlueprintCallable, Category = "ScriptableTool|Gizmos", Meta = (Identifier="Gizmo1") )
+	void SetGizmoVisible(
+		FString Identifier,
+		bool bVisible);
+
+	/**
+	 * Update the Transform on the Gizmo specified by the name Identifier
+	 * @param bUndoable if true, this transform change will be transacted into the undo/redo history, ie undoable
+	 */
+	UFUNCTION(BlueprintCallable, Category = "ScriptableTool|Gizmos", Meta = (Identifier="Gizmo1") )
+	void SetGizmoTransform(
+		FString Identifier,
+		FTransform NewTransform,
+		bool bUndoable );
+
+	/**
+	 * Get the current Transform on the Gizmo specified by the name Identifier
+	 */
+	UFUNCTION(BlueprintCallable, Category = "ScriptableTool|Gizmos", Meta = (Identifier="Gizmo1") )
+	UPARAM(DisplayName="Transform") FTransform 
+	GetGizmoTransform( FString Identifier);
+
+	/**
+	 * The OnGizmoTransformChanged event fires whenever the transform on any Gizmo created by CreateTRSGizmo() is modified.
+	 * The GizmoIdentifier can be used to disambiguate multiple active Gizmos.
+	 */
+	UFUNCTION(BlueprintImplementableEvent, Category = "ScriptableTool|Events")
+	void OnGizmoTransformChanged(const FString& GizmoIdentifier, FTransform NewTransform);
+
+protected:
+	// trying to avoid making a UStruct for this internal stuff
+	UPROPERTY(Transient, DuplicateTransient, NonTransactional, SkipSerialization)
+	TMap<FString, TObjectPtr<UCombinedTransformGizmo>> Gizmos;
+
+	virtual void OnGizmoTransformChanged_Handler(FString GizmoIdentifier, FTransform NewTransform);
+
+
+
+
+public:
+	// message API
+
+	/**
+	 * Append a Message to the UE Editor Log. 
+	 * @param bHighlighted if true, the message is emitted as a Warning, otherwise as a Log (normal)
+	 */
+	UFUNCTION(BlueprintCallable, Category = "ScriptableTool|Messaging")
+	void AddLogMessage( FText Message, bool bHighlighted = false );
+
+	/**
+	 * Display a short Help message for the user, ie to guide them in Tool usage.
+	 * In the UE5 Editor this message appears in the bottom bar of the window.
+	 */
+	UFUNCTION(BlueprintCallable, Category = "ScriptableTool|Messaging")
+	void DisplayUserHelpMessage( FText Message );
+
+	/**
+	 * Display a Warning message to the user, ie to indicate a problem/issue occurred. 
+	 */
+	UFUNCTION(BlueprintCallable, Category = "ScriptableTool|Messaging")
+	void DisplayUserWarningMessage( FText Message );
+
+	/**
+	 * Clear any active message shown via DisplayUserHelpMessage and/or DisplayUserWarningMessage
+	 */
+	UFUNCTION(BlueprintCallable, Category = "ScriptableTool|Messaging")
+	void ClearUserMessages( bool bHelpMessage, bool bWarningMessage );
+
+
+public:
+
+	// TODO: hotkey API
+};
+
+
+
+
+
+
+UCLASS(meta = (ScriptName = "ScriptableTools_Util"))
+class SCRIPTABLETOOLSFRAMEWORK_API UScriptableToolsUtilityLibrary : public UBlueprintFunctionLibrary
+{
+	GENERATED_BODY()
+public:
+
+	UFUNCTION(BlueprintPure, Category="ScriptableTool|Util")
+	static UPARAM(DisplayName = "RayHit Result") FInputRayHit 
+	MakeInputRayHit_Miss();
+
+	UFUNCTION(BlueprintPure, Category="ScriptableTool|Util")
+	static UPARAM(DisplayName = "RayHit Result") FInputRayHit 
+	MakeInputRayHit_MaxDepth();
+
+	UFUNCTION(BlueprintPure, Category="ScriptableTool|Util")
+	static UPARAM(DisplayName = "RayHit Result") FInputRayHit 
+	MakeInputRayHit(double HitDepth, UObject* OptionalHitObject);
+
+};
\ No newline at end of file
diff --git a/Engine/Plugins/Experimental/ScriptableToolsFramework/Source/ScriptableToolsFramework/Public/ScriptableToolBuilder.h b/Engine/Plugins/Experimental/ScriptableToolsFramework/Source/ScriptableToolsFramework/Public/ScriptableToolBuilder.h
new file mode 100644
--- /dev/null
+++ b/Engine/Plugins/Experimental/ScriptableToolsFramework/Source/ScriptableToolsFramework/Public/ScriptableToolBuilder.h
@@ -0,0 +1,26 @@
+// Copyright Epic Games, Inc. All Rights Reserved.
+
+#pragma once
+
+#include "InteractiveToolBuilder.h"
+#include "ScriptableInteractiveTool.h"
+#include "ScriptableToolBuilder.generated.h"
+
+
+/**
+ * UBaseScriptableToolBuilder is a trivial base UInteractiveToolBuilder for any UScriptableInteractiveTool subclass.
+ * CanBuildTool will return true as long as the ToolClass is a valid UClass.
+ */
+UCLASS()
+class SCRIPTABLETOOLSFRAMEWORK_API UBaseScriptableToolBuilder : public UInteractiveToolBuilder
+{
+	GENERATED_BODY()
+public:
+
+	TWeakObjectPtr<UClass> ToolClass;
+
+	virtual bool CanBuildTool(const FToolBuilderState& SceneState) const override;
+	virtual UInteractiveTool* BuildTool(const FToolBuilderState& SceneState) const override;
+};
+
+
diff --git a/Engine/Plugins/Experimental/ScriptableToolsFramework/Source/ScriptableToolsFramework/Public/ScriptableToolSet.h b/Engine/Plugins/Experimental/ScriptableToolsFramework/Source/ScriptableToolsFramework/Public/ScriptableToolSet.h
new file mode 100644
--- /dev/null
+++ b/Engine/Plugins/Experimental/ScriptableToolsFramework/Source/ScriptableToolsFramework/Public/ScriptableToolSet.h
@@ -0,0 +1,46 @@
+// Copyright Epic Games, Inc. All Rights Reserved.
+
+#pragma once
+
+#include "CoreMinimal.h"
+#include "ScriptableToolSet.generated.h"
+
+class UScriptableInteractiveTool;
+class UBaseScriptableToolBuilder;
+class UClass;
+
+
+/**
+ * UScriptableToolSet represents a set of UScriptableInteractiveTool types.
+ */
+UCLASS()
+class SCRIPTABLETOOLSFRAMEWORK_API UScriptableToolSet : public UObject
+{
+	GENERATED_BODY()
+
+public:
+
+	/**
+	 * Find all UScriptableInteractiveTool classes in the current project.
+	 * (Currently no support for filtering/etc)
+	 */
+	void ReinitializeScriptableTools();
+
+	/**
+	 * Allow external code to process each UScriptableInteractiveTool in the current ToolSet
+	 */
+	void ForEachScriptableTool(
+		TFunctionRef<void(UClass* ToolClass, UBaseScriptableToolBuilder* ToolBuilder)> ProcessToolFunc);
+
+protected:
+	struct FScriptableToolInfo
+	{
+		TWeakObjectPtr<UClass> ToolClass = nullptr;
+		TWeakObjectPtr<UScriptableInteractiveTool> ToolCDO;
+		TWeakObjectPtr<UBaseScriptableToolBuilder> ToolBuilder;
+	};
+	TArray<FScriptableToolInfo> Tools;
+
+	UPROPERTY()
+	TArray<TObjectPtr<UBaseScriptableToolBuilder>> ToolBuilders;
+};
diff --git a/Engine/Plugins/Experimental/ScriptableToolsFramework/Source/ScriptableToolsFramework/Public/ScriptableToolsFrameworkModule.h b/Engine/Plugins/Experimental/ScriptableToolsFramework/Source/ScriptableToolsFramework/Public/ScriptableToolsFrameworkModule.h
new file mode 100644
--- /dev/null
+++ b/Engine/Plugins/Experimental/ScriptableToolsFramework/Source/ScriptableToolsFramework/Public/ScriptableToolsFrameworkModule.h
@@ -0,0 +1,17 @@
+// Copyright Epic Games, Inc. All Rights Reserved.
+
+#pragma once
+
+#include "CoreMinimal.h"
+#include "Modules/ModuleManager.h"
+
+class FScriptableToolsFrameworkModule : public IModuleInterface
+{
+public:
+
+	/** IModuleInterface implementation */
+	virtual void StartupModule() override;
+	virtual void ShutdownModule() override;
+};
+
+
diff --git a/Engine/Plugins/Experimental/ScriptableToolsFramework/Source/ScriptableToolsFramework/ScriptableToolsFramework.Build.cs b/Engine/Plugins/Experimental/ScriptableToolsFramework/Source/ScriptableToolsFramework/ScriptableToolsFramework.Build.cs
new file mode 100644
--- /dev/null
+++ b/Engine/Plugins/Experimental/ScriptableToolsFramework/Source/ScriptableToolsFramework/ScriptableToolsFramework.Build.cs
@@ -0,0 +1,54 @@
+// Copyright Epic Games, Inc. All Rights Reserved.
+
+using UnrealBuildTool;
+
+public class ScriptableToolsFramework : ModuleRules
+{
+	public ScriptableToolsFramework(ReadOnlyTargetRules Target) : base(Target)
+	{
+		PCHUsage = ModuleRules.PCHUsageMode.UseExplicitOrSharedPCHs;
+
+		PublicIncludePaths.AddRange(
+			new string[] {
+				// ... add public include paths required here ...
+			}
+			);
+
+
+		PrivateIncludePaths.AddRange(
+			new string[] {
+				// ... add other private include paths required here ...
+			}
+			);
+
+
+		PublicDependencyModuleNames.AddRange(
+			new string[]
+			{
+				"Core",
+				"CoreUObject",
+				"PhysicsCore",
+				"RenderCore",
+				"GeometryCore",
+				"InteractiveToolsFramework"
+			}
+			);
+
+
+		PrivateDependencyModuleNames.AddRange(
+			new string[]
+			{
+				"Engine"
+			}
+			);
+
+
+		DynamicallyLoadedModuleNames.AddRange(
+			new string[]
+			{
+				// ... add any modules that your module loads dynamically here ...
+			}
+			);
+
+	}
+}
