move Editor-only code to Editor-only module.

diff --git a/Engine/Plugins/Experimental/ScriptableToolsFramework/Source/EditorScriptableToolsFramework/Private/ScriptableToolSet.cpp b/Engine/Plugins/Experimental/ScriptableToolsFramework/Source/EditorScriptableToolsFramework/Private/ScriptableToolSet.cpp
new file mode 100644
--- /dev/null
+++ b/Engine/Plugins/Experimental/ScriptableToolsFramework/Source/EditorScriptableToolsFramework/Private/ScriptableToolSet.cpp
@@ -0,0 +1,109 @@
+// Copyright Epic Games, Inc. All Rights Reserved.
+
+#include "ScriptableToolSet.h"
+
+#include "ScriptableInteractiveTool.h"
+#include "ScriptableToolBuilder.h"
+
+#include "UObject/Class.h"
+#include "UObject/ObjectMacros.h"
+#include "UObject/UObjectGlobals.h"
+#include "UObject/SoftObjectPath.h"
+#include "Templates/SubclassOf.h"
+#include "Kismet2/KismetEditorUtilities.h"
+
+#include "Modules/ModuleManager.h"
+#include "AssetRegistry/AssetRegistryModule.h"
+#include "AssetRegistry/IAssetRegistry.h"
+
+
+void UScriptableToolSet::ReinitializeScriptableTools()
+{
+	UClass* ScriptableToolClass = UScriptableInteractiveTool::StaticClass();
+	UScriptableInteractiveTool* ScriptableToolCDO = ScriptableToolClass->GetDefaultObject<UScriptableInteractiveTool>();
+
+	TSet<UClass*> PotentialToolClasses;
+
+	// Iterate over Blueprint classes to try to find UScriptableInteractiveTool blueprints
+	// Note that this code may not be fully reliable, but it appears to work so far...
+
+	// Load the asset registry module
+	FAssetRegistryModule& AssetRegistryModule = FModuleManager::LoadModuleChecked<FAssetRegistryModule>(FName("AssetRegistry"));
+	IAssetRegistry& AssetRegistry = AssetRegistryModule.Get();
+	// The asset registry is populated asynchronously at startup, so there's no guarantee it has finished.
+	// This simple approach just runs a synchronous scan on the entire content directory.
+	// Better solutions would be to specify only the path to where the relevant blueprints are,
+	// or to register a callback with the asset registry to be notified of when it's finished populating.
+	TArray< FString > ContentPaths;
+	ContentPaths.Add(TEXT("/Game"));
+	AssetRegistry.ScanPathsSynchronous(ContentPaths);
+
+	FTopLevelAssetPath BaseClassPath = ScriptableToolClass->GetClassPathName();
+
+	// Use the asset registry to get the set of all class names deriving from Base
+	TArray<FTopLevelAssetPath> BaseModeClasses;
+	BaseModeClasses.Add(BaseClassPath);
+	TSet<FTopLevelAssetPath> DerivedClassPaths;
+	AssetRegistry.GetDerivedClassNames(BaseModeClasses, TSet<FTopLevelAssetPath>(), DerivedClassPaths);
+
+	for (const FTopLevelAssetPath& ClassPath : DerivedClassPaths)
+	{
+		// don't include Base tools  (note this will also catch EditorScriptableToolsFramework)
+		if (ClassPath.ToString().Contains(TEXT("ScriptableToolsFramework")))
+		{
+			continue;
+		}
+
+		FSoftObjectPath ObjectPath(ClassPath.ToString());
+		TSoftClassPtr<UScriptableInteractiveTool> SoftClass = TSoftClassPtr<UScriptableInteractiveTool>(ObjectPath);
+		if (UClass* Class = SoftClass.LoadSynchronous())
+		{
+			if (Class->HasAnyClassFlags(CLASS_Abstract) || (Class->GetAuthoritativeClass() != Class))
+			{
+				continue;
+			}
+
+			PotentialToolClasses.Add(Class);
+		}
+	}
+
+	Tools.Reset();
+	ToolBuilders.Reset();
+
+	// if the class is viable, create a ToolBuilder for it
+	for (UClass* Class : PotentialToolClasses)
+	{
+		if (Class->IsChildOf(ScriptableToolClass))
+		{
+			FScriptableToolInfo ToolInfo;
+			ToolInfo.ToolClass = Class;
+			ToolInfo.ToolCDO = Class->GetDefaultObject<UScriptableInteractiveTool>();
+
+			UBaseScriptableToolBuilder* ToolBuilder = NewObject<UBaseScriptableToolBuilder>(this);
+			ToolBuilder->ToolClass = Class;
+			ToolBuilders.Add(ToolBuilder);
+
+			ToolInfo.ToolBuilder = ToolBuilder;
+			Tools.Add(ToolInfo);
+		}
+	}
+
+
+}
+
+
+
+
+void UScriptableToolSet::ForEachScriptableTool(
+	TFunctionRef<void(UClass* ToolClass, UBaseScriptableToolBuilder* ToolBuilder)> ProcessToolFunc)
+{
+	for (FScriptableToolInfo& ToolInfo : Tools)
+	{
+		if (ToolInfo.ToolClass.IsValid() && ToolInfo.ToolBuilder.IsValid())
+		{
+			ProcessToolFunc(ToolInfo.ToolClass.Get(), ToolInfo.ToolBuilder.Get());
+		}
+	}
+}
+
+
diff --git a/Engine/Plugins/Experimental/ScriptableToolsFramework/Source/EditorScriptableToolsFramework/Public/ScriptableToolSet.h b/Engine/Plugins/Experimental/ScriptableToolsFramework/Source/EditorScriptableToolsFramework/Public/ScriptableToolSet.h
new file mode 100644
--- /dev/null
+++ b/Engine/Plugins/Experimental/ScriptableToolsFramework/Source/EditorScriptableToolsFramework/Public/ScriptableToolSet.h
@@ -0,0 +1,46 @@
+// Copyright Epic Games, Inc. All Rights Reserved.
+
+#pragma once
+
+#include "CoreMinimal.h"
+#include "ScriptableToolSet.generated.h"
+
+class UScriptableInteractiveTool;
+class UBaseScriptableToolBuilder;
+class UClass;
+
+
+/**
+ * UScriptableToolSet represents a set of UScriptableInteractiveTool types.
+ */
+UCLASS()
+class EDITORSCRIPTABLETOOLSFRAMEWORK_API UScriptableToolSet : public UObject
+{
+	GENERATED_BODY()
+
+public:
+
+	/**
+	 * Find all UScriptableInteractiveTool classes in the current project.
+	 * (Currently no support for filtering/etc)
+	 */
+	void ReinitializeScriptableTools();
+
+	/**
+	 * Allow external code to process each UScriptableInteractiveTool in the current ToolSet
+	 */
+	void ForEachScriptableTool(
+		TFunctionRef<void(UClass* ToolClass, UBaseScriptableToolBuilder* ToolBuilder)> ProcessToolFunc);
+
+protected:
+	struct FScriptableToolInfo
+	{
+		TWeakObjectPtr<UClass> ToolClass = nullptr;
+		TWeakObjectPtr<UScriptableInteractiveTool> ToolCDO;
+		TWeakObjectPtr<UBaseScriptableToolBuilder> ToolBuilder;
+	};
+	TArray<FScriptableToolInfo> Tools;
+
+	UPROPERTY()
+	TArray<TObjectPtr<UBaseScriptableToolBuilder>> ToolBuilders;
+};
diff --git a/Engine/Plugins/Experimental/ScriptableToolsFramework/Source/ScriptableToolsFramework/Private/ScriptableToolSet.cpp b/Engine/Plugins/Experimental/ScriptableToolsFramework/Source/ScriptableToolsFramework/Private/ScriptableToolSet.cpp
deleted file mode 100644
diff --git a/Engine/Plugins/Experimental/ScriptableToolsFramework/Source/ScriptableToolsFramework/Public/ScriptableToolSet.h b/Engine/Plugins/Experimental/ScriptableToolsFramework/Source/ScriptableToolsFramework/Public/ScriptableToolSet.h
deleted file mode 100644
