ScriptableTools improvements.
Change EToolsFrameworkOutcomePins to be enum class : uint8 to avoid naming conflicts
Add UScriptableTool_RenderAPI::DrawRectWidthHeightXY
Add UScriptableInteractiveTool::bShowToolInEditor flag, if disabled, ScriptableToolsEditorMode will not show the Tool in the Palette
ScriptableToolsEditorMode now refreshes palette on global OnBlueprintCompiled delegate, to catch tool renaming/etc

diff --git a/Engine/Plugins/Experimental/ScriptableToolsEditorMode/Source/ScriptableToolsEditorMode/Private/ScriptableToolsEditorMode.cpp b/Engine/Plugins/Experimental/ScriptableToolsEditorMode/Source/ScriptableToolsEditorMode/Private/ScriptableToolsEditorMode.cpp
--- a/Engine/Plugins/Experimental/ScriptableToolsEditorMode/Source/ScriptableToolsEditorMode/Private/ScriptableToolsEditorMode.cpp
+++ b/Engine/Plugins/Experimental/ScriptableToolsEditorMode/Source/ScriptableToolsEditorMode/Private/ScriptableToolsEditorMode.cpp
@@ -167,6 +167,7 @@ void UScriptableToolsEditorMode::Enter()
 	GetToolManager()->SelectActiveToolType(EToolSide::Left, TEXT("BeginMeshInspectorTool"));
 
 	BlueprintPreCompileHandle = GEditor->OnBlueprintPreCompile().AddUObject(this, &UScriptableToolsEditorMode::OnBlueprintPreCompile); 
+	BlueprintCompiledHandle = GEditor->OnBlueprintCompiled().AddUObject(this, &UScriptableToolsEditorMode::OnBlueprintCompiled); 
 
 	// do any toolkit UI initialization that depends on the mode setup above
 	if (Toolkit.IsValid())
@@ -185,7 +186,11 @@ void UScriptableToolsEditorMode::OnBlueprintPreCompile(UBlueprint* Blueprint)
 	{
 		GetToolManager()->DeactivateTool(EToolSide::Left, EToolShutdownType::Cancel);
 	}
+}
+
 
+void UScriptableToolsEditorMode::OnBlueprintCompiled()
+{
 	// Probably not necessary to always rebuild the palette here. But currently this lets us respond
 	// to changes in the tool name/setting/etc
 	if (Toolkit.IsValid())
@@ -195,7 +200,6 @@ void UScriptableToolsEditorMode::OnBlueprintPreCompile(UBlueprint* Blueprint)
 	}
 }
 
-
 void UScriptableToolsEditorMode::Exit()
 {
 	GEditor->OnBlueprintPreCompile().Remove(BlueprintPreCompileHandle);
diff --git a/Engine/Plugins/Experimental/ScriptableToolsEditorMode/Source/ScriptableToolsEditorMode/Private/ScriptableToolsEditorModeToolkit.cpp b/Engine/Plugins/Experimental/ScriptableToolsEditorMode/Source/ScriptableToolsEditorMode/Private/ScriptableToolsEditorModeToolkit.cpp
--- a/Engine/Plugins/Experimental/ScriptableToolsEditorMode/Source/ScriptableToolsEditorMode/Private/ScriptableToolsEditorModeToolkit.cpp
+++ b/Engine/Plugins/Experimental/ScriptableToolsEditorMode/Source/ScriptableToolsEditorMode/Private/ScriptableToolsEditorModeToolkit.cpp
@@ -302,6 +302,11 @@ void FScriptableToolsEditorModeToolkit::GetToolPaletteNames(TArray<FName>& Palet
 	ScriptableTools->ForEachScriptableTool( [&](UClass* ToolClass, UBaseScriptableToolBuilder* ToolBuilder) 
 	{
 		UScriptableInteractiveTool* ToolCDO = Cast<UScriptableInteractiveTool>(ToolClass->GetDefaultObject());
+		if (ToolCDO->bShowToolInEditor == false)
+		{
+			return;
+		}
+
 		if (ToolCDO->ToolCategory.IsEmpty() == false)
 		{
 			FName CategoryName(ToolCDO->ToolCategory.ToString());
@@ -353,6 +358,11 @@ void FScriptableToolsEditorModeToolkit::BuildToolPalette(FName PaletteIndex, cla
 	ScriptableTools->ForEachScriptableTool([&](UClass* ToolClass, UBaseScriptableToolBuilder* ToolBuilder)
 	{
 		UScriptableInteractiveTool* ToolCDO = Cast<UScriptableInteractiveTool>(ToolClass->GetDefaultObject());
+		if (ToolCDO->bShowToolInEditor == false)
+		{
+			return;
+		}
+
 		FName UseCategoryName = ToolCDO->ToolCategory.IsEmpty() ? CustomToolsTabName : FName(ToolCDO->ToolCategory.ToString());
 		if (ActiveToolCategories.Contains(UseCategoryName) == false)
 		{
diff --git a/Engine/Plugins/Experimental/ScriptableToolsEditorMode/Source/ScriptableToolsEditorMode/Public/ScriptableToolsEditorMode.h b/Engine/Plugins/Experimental/ScriptableToolsEditorMode/Source/ScriptableToolsEditorMode/Public/ScriptableToolsEditorMode.h
--- a/Engine/Plugins/Experimental/ScriptableToolsEditorMode/Source/ScriptableToolsEditorMode/Public/ScriptableToolsEditorMode.h
+++ b/Engine/Plugins/Experimental/ScriptableToolsEditorMode/Source/ScriptableToolsEditorMode/Public/ScriptableToolsEditorMode.h
@@ -84,6 +84,9 @@ protected:
 	void OnBlueprintPreCompile(UBlueprint* Blueprint);
 	FDelegateHandle BlueprintPreCompileHandle;
 
+	void OnBlueprintCompiled();
+	FDelegateHandle BlueprintCompiledHandle;
+
 protected:
 	UPROPERTY()
 	TObjectPtr<UScriptableToolSet> ScriptableTools;
diff --git a/Engine/Plugins/Experimental/ScriptableToolsFramework/Source/ScriptableToolsFramework/Private/ScriptableInteractiveTool.cpp b/Engine/Plugins/Experimental/ScriptableToolsFramework/Source/ScriptableToolsFramework/Private/ScriptableInteractiveTool.cpp
--- a/Engine/Plugins/Experimental/ScriptableToolsFramework/Source/ScriptableToolsFramework/Private/ScriptableInteractiveTool.cpp
+++ b/Engine/Plugins/Experimental/ScriptableToolsFramework/Source/ScriptableToolsFramework/Private/ScriptableInteractiveTool.cpp
@@ -22,7 +22,7 @@
 
 
 
-UScriptableInteractiveTool* UScriptableInteractiveToolPropertySet::GetOwningTool(TEnumAsByte<EToolsFrameworkOutcomePins>& Outcome)
+UScriptableInteractiveTool* UScriptableInteractiveToolPropertySet::GetOwningTool(EToolsFrameworkOutcomePins& Outcome)
 {
 	if (ParentTool.IsValid())
 	{
@@ -50,6 +50,44 @@ UScriptableTool_RenderAPI* UScriptableTool_RenderAPI::DrawLine(FVector Start, FV
 }
 
 
+UScriptableTool_RenderAPI* UScriptableTool_RenderAPI::DrawRectWidthHeightXY(FTransform Transform, double Width, double Height, FLinearColor Color, float LineThickness, float DepthBias, bool bDepthTested, bool bOriginIsCenter)
+{
+	if (ActiveVisualizer == nullptr)
+	{
+		UE_LOG(LogTemp, Warning, TEXT("[ScriptableInteractiveTool::DrawRectWidthHeight] DrawRectWidthHeight can only be called during OnScriptRender event callback"));
+		return this;
+	}
+
+	ActiveVisualizer->DepthBias = DepthBias;
+
+	ActiveVisualizer->PushTransform(Transform);
+
+	if (bOriginIsCenter)
+	{
+		FVector DeltaX(Width*0.5, 0, 0), DeltaY(0,Height*0.5,0);
+		FVector A(-DeltaX-DeltaY), B(DeltaX-DeltaY), C(DeltaX+DeltaY), D(-DeltaX+DeltaY);
+		ActiveVisualizer->DrawLine(A, B, Color, LineThickness, bDepthTested);
+		ActiveVisualizer->DrawLine(B, C, Color, LineThickness, bDepthTested);
+		ActiveVisualizer->DrawLine(C, D, Color, LineThickness, bDepthTested);
+		ActiveVisualizer->DrawLine(D, A, Color, LineThickness, bDepthTested);
+	}
+	else
+	{
+		FVector DeltaX(Width, 0, 0), DeltaY(0,Height,0);
+		FVector A(0,0,0), B(DeltaX), C(DeltaX+DeltaY), D(DeltaY);
+		ActiveVisualizer->DrawLine(A, B, Color, LineThickness, bDepthTested);
+		ActiveVisualizer->DrawLine(B, C, Color, LineThickness, bDepthTested);
+		ActiveVisualizer->DrawLine(C, D, Color, LineThickness, bDepthTested);
+		ActiveVisualizer->DrawLine(D, A, Color, LineThickness, bDepthTested);
+	}
+
+	ActiveVisualizer->PopTransform();
+
+	return this;
+}
+
+
+
 UScriptableTool_HUDAPI* UScriptableTool_HUDAPI::DrawTextAtLocation(FVector Location, FString String, FLinearColor Color, bool bCentered, float ShiftRowsY)
 {
 	if (ActiveCanvas == nullptr || ActiveSceneView == nullptr)
@@ -279,7 +317,7 @@ void UScriptableInteractiveTool::DrawHUD(FCanvas* Canvas, IToolsContextRenderAPI
 UScriptableInteractiveToolPropertySet* UScriptableInteractiveTool::AddPropertySetOfType(
 	TSubclassOf<UScriptableInteractiveToolPropertySet> PropertySetType,
 	FString Identifier,
-	TEnumAsByte<EToolsFrameworkOutcomePins>& Outcome)
+	EToolsFrameworkOutcomePins& Outcome)
 {
 	Outcome = EToolsFrameworkOutcomePins::Failure;
 
@@ -320,7 +358,7 @@ UScriptableInteractiveToolPropertySet* UScriptableInteractiveTool::AddPropertySe
 
 void UScriptableInteractiveTool::RemovePropertySetByName(
 	FString Identifier,
-	TEnumAsByte<EToolsFrameworkOutcomePins>& Outcome)
+	EToolsFrameworkOutcomePins& Outcome)
 {
 	Outcome = EToolsFrameworkOutcomePins::Failure;
 
@@ -359,7 +397,7 @@ void UScriptableInteractiveTool::ForcePropertySetUpdateByName(
 }
 
 
-void UScriptableInteractiveTool::RestorePropertySetSettings(
+UScriptableInteractiveToolPropertySet* UScriptableInteractiveTool::RestorePropertySetSettings(
 	UScriptableInteractiveToolPropertySet* PropertySet,
 	FString SaveKey)
 {
@@ -367,10 +405,11 @@ void UScriptableInteractiveTool::RestorePropertySetSettings(
 	{
 		PropertySet->RestoreProperties(this, SaveKey);
 	}
+	return PropertySet;
 }
 
 
-void UScriptableInteractiveTool::SavePropertySetSettings(
+UScriptableInteractiveToolPropertySet* UScriptableInteractiveTool::SavePropertySetSettings(
 	UScriptableInteractiveToolPropertySet* PropertySet,
 	FString SaveKey)
 {
@@ -378,8 +417,13 @@ void UScriptableInteractiveTool::SavePropertySetSettings(
 	{
 		PropertySet->SaveProperties(this, SaveKey);
 	}
+	return PropertySet;
 }
 
+
+
+
+
 namespace UELocal
 {
 
@@ -420,7 +464,7 @@ PropertyType* FindValidPropertyByTypeAndName(
 } // end namespace UELocal
 
 
-void UScriptableInteractiveTool::WatchFloatProperty(
+UScriptableInteractiveToolPropertySet* UScriptableInteractiveTool::WatchFloatProperty(
 	UScriptableInteractiveToolPropertySet* PropertySet,
 	FString PropertyName,
 	const FToolFloatPropertyModifiedDelegate& OnModified)
@@ -431,7 +475,7 @@ void UScriptableInteractiveTool::WatchFloatProperty(
 		if (NumericProp->IsFloatingPoint() == false)
 		{
 			UE_LOG(LogTemp, Warning, TEXT("WatchFloatProperty: Property %s is not Float type"), *PropertyName);
-			return;
+			return PropertySet;
 		}
 
 		int32 ArrayIndex = WatchFloatPropertyDelegates.Num();
@@ -447,12 +491,12 @@ void UScriptableInteractiveTool::WatchFloatProperty(
 				WatchFloatPropertyDelegates[ArrayIndex].ExecuteIfBound(PropertySet, PropertyName, NewValue);
 			} );
 	}
-
+	return PropertySet;
 }
 
 
 
-void UScriptableInteractiveTool::WatchIntProperty(
+UScriptableInteractiveToolPropertySet* UScriptableInteractiveTool::WatchIntProperty(
 	UScriptableInteractiveToolPropertySet* PropertySet,
 	FString PropertyName,
 	const FToolIntPropertyModifiedDelegate& OnModified)
@@ -463,7 +507,7 @@ void UScriptableInteractiveTool::WatchIntProperty(
 		if (NumericProp->IsInteger() == false)
 		{
 			UE_LOG(LogTemp, Warning, TEXT("WatchIntProperty: Property %s is not Integer type"), *PropertyName);
-			return;
+			return PropertySet;
 		}
 
 		int32 ArrayIndex = WatchIntPropertyDelegates.Num();
@@ -478,11 +522,12 @@ void UScriptableInteractiveTool::WatchIntProperty(
 				WatchIntPropertyDelegates[ArrayIndex].ExecuteIfBound(PropertySet, PropertyName, NewValue);
 			} );
 	}
+	return PropertySet;
 }
 
 
 
-void UScriptableInteractiveTool::WatchBoolProperty(
+UScriptableInteractiveToolPropertySet* UScriptableInteractiveTool::WatchBoolProperty(
 	UScriptableInteractiveToolPropertySet* PropertySet,
 	FString PropertyName,
 	const FToolBoolPropertyModifiedDelegate& OnModified)
@@ -503,10 +548,11 @@ void UScriptableInteractiveTool::WatchBoolProperty(
 				WatchBoolPropertyDelegates[ArrayIndex].ExecuteIfBound(PropertySet, PropertyName, bNewValue);
 			} );
 	}
+	return PropertySet;
 }
 
 
-void UScriptableInteractiveTool::WatchEnumProperty(
+UScriptableInteractiveToolPropertySet* UScriptableInteractiveTool::WatchEnumProperty(
 	UScriptableInteractiveToolPropertySet* PropertySet,
 	FString PropertyName,
 	const FToolEnumPropertyModifiedDelegate& OnModified)
@@ -528,11 +574,12 @@ void UScriptableInteractiveTool::WatchEnumProperty(
 				WatchEnumPropertyDelegates[ArrayIndex].ExecuteIfBound(PropertySet, PropertyName, (uint8)NewValue);
 			} );
 	}
+	return PropertySet;
 }
 
 
 
-void UScriptableInteractiveTool::WatchStringProperty(
+UScriptableInteractiveToolPropertySet* UScriptableInteractiveTool::WatchStringProperty(
 	UScriptableInteractiveToolPropertySet* PropertySet,
 	FString PropertyName,
 	const FToolStringPropertyModifiedDelegate& OnModified)
@@ -552,10 +599,11 @@ void UScriptableInteractiveTool::WatchStringProperty(
 				WatchStringPropertyDelegates[ArrayIndex].ExecuteIfBound(PropertySet, PropertyName, NewValue);
 			} );
 	}
+	return PropertySet;
 }
 
 
-void UScriptableInteractiveTool::WatchNameProperty(
+UScriptableInteractiveToolPropertySet* UScriptableInteractiveTool::WatchNameProperty(
 	UScriptableInteractiveToolPropertySet* PropertySet,
 	FString PropertyName,
 	const FToolFNamePropertyModifiedDelegate& OnModified)
@@ -575,10 +623,11 @@ void UScriptableInteractiveTool::WatchNameProperty(
 				WatchFNamePropertyDelegates[ArrayIndex].ExecuteIfBound(PropertySet, PropertyName, NewValue);
 			} );
 	}
+	return PropertySet;
 }
 
 
-void UScriptableInteractiveTool::WatchObjectProperty(
+UScriptableInteractiveToolPropertySet* UScriptableInteractiveTool::WatchObjectProperty(
 	UScriptableInteractiveToolPropertySet* PropertySet,
 	FString PropertyName,
 	const FToolObjectPropertyModifiedDelegate& OnModified)
@@ -598,12 +647,13 @@ void UScriptableInteractiveTool::WatchObjectProperty(
 				WatchObjectPropertyDelegates[ArrayIndex].ExecuteIfBound(PropertySet, PropertyName, NewValue);
 			} );
 	}
+	return PropertySet;
 }
 
 
 
 
-void UScriptableInteractiveTool::WatchProperty(
+UScriptableInteractiveToolPropertySet* UScriptableInteractiveTool::WatchProperty(
 	UScriptableInteractiveToolPropertySet* PropertySet,
 	FString PropertyName,
 	const FToolPropertyModifiedDelegate& OnModified)
@@ -612,7 +662,7 @@ void UScriptableInteractiveTool::WatchProperty(
 			PropertySet, PropertyName, TEXT("WatchProperty"));
 	if (Property == nullptr)
 	{
-		return;
+		return PropertySet;
 	}
 
 	int32 ArrayIndex = WatchAnyPropertyInfo.Num();
@@ -768,7 +818,7 @@ void UScriptableInteractiveTool::WatchProperty(
 		//		WatchAnyPropertyInfo[ArrayIndex].Delegate.ExecuteIfBound(PropertySet, PropertyName);
 		//	} );
 	}
-
+	return PropertySet;
 }
 
 
@@ -778,7 +828,7 @@ void UScriptableInteractiveTool::CreateTRSGizmo(
 	FString Identifier,
 	FTransform InitialTransform,
 	FScriptableToolGizmoOptions GizmoOptions,
-	TEnumAsByte<EToolsFrameworkOutcomePins>& Outcome)
+	EToolsFrameworkOutcomePins& Outcome)
 {
 	Outcome = EToolsFrameworkOutcomePins::Failure;
 
@@ -889,7 +939,7 @@ void UScriptableInteractiveTool::CreateTRSGizmo(
 
 void UScriptableInteractiveTool::DestroyTRSGizmo(
 	FString Identifier,
-	TEnumAsByte<EToolsFrameworkOutcomePins>& Outcome)
+	EToolsFrameworkOutcomePins& Outcome)
 {
 	Outcome = EToolsFrameworkOutcomePins::Failure;
 
diff --git a/Engine/Plugins/Experimental/ScriptableToolsFramework/Source/ScriptableToolsFramework/Public/ScriptableInteractiveTool.h b/Engine/Plugins/Experimental/ScriptableToolsFramework/Source/ScriptableToolsFramework/Public/ScriptableInteractiveTool.h
--- a/Engine/Plugins/Experimental/ScriptableToolsFramework/Source/ScriptableToolsFramework/Public/ScriptableInteractiveTool.h
+++ b/Engine/Plugins/Experimental/ScriptableToolsFramework/Source/ScriptableToolsFramework/Public/ScriptableInteractiveTool.h
@@ -16,7 +16,7 @@ class FCanvas;
 class FSceneView;
 
 UENUM(BlueprintType)
-enum EToolsFrameworkOutcomePins
+enum class EToolsFrameworkOutcomePins : uint8
 {
 	Success,
 	Failure
@@ -170,6 +170,11 @@ public:
 	UPARAM(DisplayName="Render Object") UScriptableTool_RenderAPI* 
 	DrawLine(FVector Start, FVector End, FLinearColor Color, float Thickness = 1.0, float DepthBias = 0.0f, bool bDepthTested = true);
 
+	UFUNCTION(BlueprintCallable, Category = "ScriptableTool|Render")
+	UPARAM(DisplayName="Render Object") UScriptableTool_RenderAPI* 
+	DrawRectWidthHeightXY(FTransform Transform, double Width, double Height, FLinearColor Color, float LineThickness = 1.0, float DepthBias = 0.0f, bool bDepthTested = true, bool bCentered = true);
+
+
 	// DrawPoint, DrawCircle, DrawViewFacingCircle
 	// DrawWireCylinder, DrawWireBox, DrawSquare
 	// PushTransform, PopTransform, SetTransform, PopAllTransforms
@@ -238,7 +243,7 @@ class SCRIPTABLETOOLSFRAMEWORK_API UScriptableInteractiveToolPropertySet : publi
 	/** Access the Tool that owns this PropertySet */
 	UFUNCTION(BlueprintCallable, Category = "ScriptableTool", Meta = (ExpandEnumAsExecs = "Outcome"))
 	UScriptableInteractiveTool* GetOwningTool(
-		TEnumAsByte<EToolsFrameworkOutcomePins>& Outcome);
+		EToolsFrameworkOutcomePins& Outcome);
 
 protected:
 	TWeakObjectPtr<UScriptableInteractiveTool> ParentTool;
@@ -283,6 +288,10 @@ public:
 	UPROPERTY(BlueprintReadWrite, EditAnywhere, Category = "Scriptable Tool Settings", meta=(DisplayName="Tooltip"))
 	FText ToolTooltip;
 
+	/** A generic flag to indicate whether this Tool should be shown in the UE Editor. This may be interpreted in different ways */
+	UPROPERTY(BlueprintReadWrite, EditAnywhere, Category = "Scriptable Tool Settings", meta=(DisplayName="Visible in Editor"))
+	bool bShowToolInEditor = true;
+
 	/** Specifies how the user exits this Tool, either Accept/Cancel-style or Complete-style */
 	UPROPERTY(BlueprintReadWrite, EditAnywhere, Category = "Scriptable Tool Settings", meta=(DisplayName="Shutdown Type"))
 	EScriptableToolShutdownType ToolShutdownType = EScriptableToolShutdownType::Complete;
@@ -418,7 +427,7 @@ public:
 	AddPropertySetOfType(
 		TSubclassOf<UScriptableInteractiveToolPropertySet> PropertySetType,
 		FString Identifier,
-		TEnumAsByte<EToolsFrameworkOutcomePins>& Outcome);
+		EToolsFrameworkOutcomePins& Outcome);
 
 	/**
 	 * Remove a Property Set from the current Tool, found via it's unique Identifier.
@@ -427,7 +436,7 @@ public:
 	UFUNCTION(BlueprintCallable, Category = "ScriptableTool|PropertySets", Meta = (ExpandEnumAsExecs = "Outcome", Identifier="Settings"))
 	void RemovePropertySetByName(
 		FString Identifier,
-		TEnumAsByte<EToolsFrameworkOutcomePins>& Outcome);
+		EToolsFrameworkOutcomePins& Outcome);
 
 	/**
 	 * Set the visibility of a Property Set that is paired with the given unique Identifier.
@@ -453,7 +462,7 @@ public:
 	 * in the same Engine session (ie in a previous invocation of the Tool, or another Tool that uses the same Property Set)
 	 */
 	UFUNCTION(BlueprintCallable, Category = "ScriptableTool|PropertySets", Meta = (Identifier="Settings"))
-	void RestorePropertySetSettings(
+	UPARAM(DisplayName="Property Set") UScriptableInteractiveToolPropertySet* RestorePropertySetSettings(
 		UScriptableInteractiveToolPropertySet* PropertySet,
 		FString SaveKey);
 
@@ -462,7 +471,7 @@ public:
 	 * These saved values can be restored in future Tool invocations based on the SaveKey.
 	 */
 	UFUNCTION(BlueprintCallable, Category = "ScriptableTool|PropertySets", Meta = (Identifier="Settings"))
-	void SavePropertySetSettings(
+	UPARAM(DisplayName="Property Set") UScriptableInteractiveToolPropertySet* SavePropertySetSettings(
 		UScriptableInteractiveToolPropertySet* PropertySet,
 		FString SaveKey);
 
@@ -479,7 +488,7 @@ public:
 	 * @param OnModified this delegate will be called if the Property value changes.
 	 */
 	UFUNCTION(BlueprintCallable, Category = "ScriptableTool|PropertySets")
-	void WatchFloatProperty(
+	UPARAM(DisplayName="Property Set") UScriptableInteractiveToolPropertySet* WatchFloatProperty(
 		UScriptableInteractiveToolPropertySet* PropertySet,
 		FString PropertyName,
 		const FToolFloatPropertyModifiedDelegate& OnModified );
@@ -491,7 +500,7 @@ public:
 	 * @param OnModified this delegate will be called if the Property value changes.
 	 */
 	UFUNCTION(BlueprintCallable, Category = "ScriptableTool|PropertySets")
-	void WatchIntProperty(
+	UPARAM(DisplayName="Property Set") UScriptableInteractiveToolPropertySet* WatchIntProperty(
 		UScriptableInteractiveToolPropertySet* PropertySet,
 		FString PropertyName,
 		const FToolIntPropertyModifiedDelegate& OnModified );
@@ -503,7 +512,7 @@ public:
 	 * @param OnModified this delegate will be called if the Property value changes.
 	 */
 	UFUNCTION(BlueprintCallable, Category = "ScriptableTool|PropertySets")
-	void WatchBoolProperty(
+	UPARAM(DisplayName="Property Set") UScriptableInteractiveToolPropertySet* WatchBoolProperty(
 		UScriptableInteractiveToolPropertySet* PropertySet,
 		FString PropertyName,
 		const FToolBoolPropertyModifiedDelegate& OnModified );
@@ -516,7 +525,7 @@ public:
 	 * @param OnModified this delegate will be called if the Property value changes.
 	 */
 	UFUNCTION(BlueprintCallable, Category = "ScriptableTool|PropertySets")
-	void WatchEnumProperty(
+	UPARAM(DisplayName="Property Set") UScriptableInteractiveToolPropertySet* WatchEnumProperty(
 		UScriptableInteractiveToolPropertySet* PropertySet,
 		FString PropertyName,
 		const FToolEnumPropertyModifiedDelegate& OnModified );
@@ -528,7 +537,7 @@ public:
 	 * @param OnModified this delegate will be called if the Property value changes.
 	 */
 	UFUNCTION(BlueprintCallable, Category = "ScriptableTool|PropertySets")
-	void WatchStringProperty(
+	UPARAM(DisplayName="Property Set") UScriptableInteractiveToolPropertySet* WatchStringProperty(
 		UScriptableInteractiveToolPropertySet* PropertySet,
 		FString PropertyName,
 		const FToolStringPropertyModifiedDelegate& OnModified );
@@ -540,7 +549,7 @@ public:
 	 * @param OnModified this delegate will be called if the Property value changes.
 	 */
 	UFUNCTION(BlueprintCallable, Category = "ScriptableTool|PropertySets")
-	void WatchNameProperty(
+	UPARAM(DisplayName="Property Set") UScriptableInteractiveToolPropertySet* WatchNameProperty(
 		UScriptableInteractiveToolPropertySet* PropertySet,
 		FString PropertyName,
 		const FToolFNamePropertyModifiedDelegate& OnModified );
@@ -552,7 +561,7 @@ public:
 	 * @param OnModified this delegate will be called if the Property value changes.
 	 */
 	UFUNCTION(BlueprintCallable, Category = "ScriptableTool|PropertySets")
-	void WatchObjectProperty(
+	UPARAM(DisplayName="Property Set") UScriptableInteractiveToolPropertySet* WatchObjectProperty(
 		UScriptableInteractiveToolPropertySet* PropertySet,
 		FString PropertyName,
 		const FToolObjectPropertyModifiedDelegate& OnModified );
@@ -570,7 +579,7 @@ public:
 	 * @param OnModified this delegate will be called if the Property value changes.
 	 */
 	UFUNCTION(BlueprintCallable, Category = "ScriptableTool|PropertySets")
-	void WatchProperty(
+	UPARAM(DisplayName="Property Set") UScriptableInteractiveToolPropertySet* WatchProperty(
 		UScriptableInteractiveToolPropertySet* PropertySet,
 		FString PropertyName,
 		const FToolPropertyModifiedDelegate& OnModified );
@@ -627,7 +636,7 @@ public:
 		FString Identifier,
 		FTransform InitialTransform,
 		FScriptableToolGizmoOptions GizmoOptions,
-		TEnumAsByte<EToolsFrameworkOutcomePins>& Outcome);
+		EToolsFrameworkOutcomePins& Outcome);
 
 	/**
 	 * Destroy a created Gizmo by name Identifier
@@ -635,7 +644,7 @@ public:
 	UFUNCTION(BlueprintCallable, Category = "ScriptableTool|Gizmos", Meta = (ExpandEnumAsExecs = "Outcome", Identifier="Gizmo1"))
 	void DestroyTRSGizmo(
 		FString Identifier,
-		TEnumAsByte<EToolsFrameworkOutcomePins>& Outcome);
+		EToolsFrameworkOutcomePins& Outcome);
 
 	/**
 	 * Set an existing Gizmo visible/hidden based on its name Identifier
@@ -706,7 +715,7 @@ public:
 	 * Clear any active message shown via DisplayUserHelpMessage and/or DisplayUserWarningMessage
 	 */
 	UFUNCTION(BlueprintCallable, Category = "ScriptableTool|Messaging")
-	void ClearUserMessages( bool bHelpMessage, bool bWarningMessage );
+	void ClearUserMessages( bool bHelpMessage = true, bool bWarningMessage = true );
 
 
 public:
