GeometryScripting: make plugin modules that are usable at Runtime load PreDefault, to ensure that if users reference BP functions from their Game Mode, the plugins are loaded first.

diff --git a/Engine/Plugins/Experimental/GeometryScripting/GeometryScripting.uplugin b/Engine/Plugins/Experimental/GeometryScripting/GeometryScripting.uplugin
--- a/Engine/Plugins/Experimental/GeometryScripting/GeometryScripting.uplugin
+++ b/Engine/Plugins/Experimental/GeometryScripting/GeometryScripting.uplugin
@@ -17,7 +17,7 @@
 		{
 			"Name": "GeometryScriptingCore",
 			"Type": "Runtime",
-			"LoadingPhase": "Default"
+			"LoadingPhase": "PreDefault"
 		},
 		{
 			"Name": "GeometryScriptingEditor",
diff --git a/Engine/Plugins/Experimental/ScriptableToolsFramework/ScriptableToolsFramework.uplugin b/Engine/Plugins/Experimental/ScriptableToolsFramework/ScriptableToolsFramework.uplugin
--- a/Engine/Plugins/Experimental/ScriptableToolsFramework/ScriptableToolsFramework.uplugin
+++ b/Engine/Plugins/Experimental/ScriptableToolsFramework/ScriptableToolsFramework.uplugin
@@ -17,7 +17,7 @@
 		{
 			"Name": "ScriptableToolsFramework",
 			"Type": "Runtime",
-			"LoadingPhase": "Default"
+			"LoadingPhase": "PreDefault"
 		},
 		{
 			"Name": "EditorScriptableToolsFramework",
