ScriptableTools: add OnGizmoTransformStateChange callback, gets begin/end and undo/redo transform events. Add support for providing custom UBaseScriptableToolBuilder for a subclass of UScriptableInteractiveTool (in C++), allows custom extended ScriptableTool C++ base classes to provide their own base ToolBuilder. Add editor-only UEditorScriptableInteractiveToolPropertySet subclass of UScriptableInteractiveToolPropertySet (necessary to get ToolPropertySet to show up when creating an Editor Blueprint).

diff --git a/Engine/Plugins/Experimental/ScriptableToolsFramework/Source/EditorScriptableToolsFramework/Private/ScriptableToolSet.cpp b/Engine/Plugins/Experimental/ScriptableToolsFramework/Source/EditorScriptableToolsFramework/Private/ScriptableToolSet.cpp
--- a/Engine/Plugins/Experimental/ScriptableToolsFramework/Source/EditorScriptableToolsFramework/Private/ScriptableToolSet.cpp
+++ b/Engine/Plugins/Experimental/ScriptableToolsFramework/Source/EditorScriptableToolsFramework/Private/ScriptableToolSet.cpp
@@ -72,7 +72,12 @@ void UScriptableToolSet::ReinitializeScriptableTools()
 			ToolInfo.ToolClass = Class;
 			ToolInfo.ToolCDO = Class->GetDefaultObject<UScriptableInteractiveTool>();
 
-			UBaseScriptableToolBuilder* ToolBuilder = NewObject<UBaseScriptableToolBuilder>(this);
+			UBaseScriptableToolBuilder* ToolBuilder = ToolInfo.ToolCDO->GetNewCustomToolBuilderInstance(this);
+			if (ToolBuilder == nullptr)
+			{
+				ToolBuilder = NewObject<UBaseScriptableToolBuilder>(this);
+			}
+
 			ToolBuilder->ToolClass = Class;
 			ToolBuilders.Add(ToolBuilder);
 
diff --git a/Engine/Plugins/Experimental/ScriptableToolsFramework/Source/EditorScriptableToolsFramework/Public/EditorScriptableInteractiveTool.h b/Engine/Plugins/Experimental/ScriptableToolsFramework/Source/EditorScriptableToolsFramework/Public/EditorScriptableInteractiveTool.h
--- a/Engine/Plugins/Experimental/ScriptableToolsFramework/Source/EditorScriptableToolsFramework/Public/EditorScriptableInteractiveTool.h
+++ b/Engine/Plugins/Experimental/ScriptableToolsFramework/Source/EditorScriptableToolsFramework/Public/EditorScriptableInteractiveTool.h
@@ -14,4 +14,13 @@ class EDITORSCRIPTABLETOOLSFRAMEWORK_API UEditorScriptableInteractiveTool : publ
 	GENERATED_BODY()
 public:
 
-};
\ No newline at end of file
+};
+
+
+UCLASS(Transient, Blueprintable)
+class EDITORSCRIPTABLETOOLSFRAMEWORK_API UEditorScriptableInteractiveToolPropertySet : public UScriptableInteractiveToolPropertySet
+{
+	GENERATED_BODY()
+public:
+
+};
diff --git a/Engine/Plugins/Experimental/ScriptableToolsFramework/Source/ScriptableToolsFramework/Private/ScriptableInteractiveTool.cpp b/Engine/Plugins/Experimental/ScriptableToolsFramework/Source/ScriptableToolsFramework/Private/ScriptableInteractiveTool.cpp
--- a/Engine/Plugins/Experimental/ScriptableToolsFramework/Source/ScriptableToolsFramework/Private/ScriptableInteractiveTool.cpp
+++ b/Engine/Plugins/Experimental/ScriptableToolsFramework/Source/ScriptableToolsFramework/Private/ScriptableInteractiveTool.cpp
@@ -923,6 +923,19 @@ void UScriptableInteractiveTool::CreateTRSGizmo(
 	{
 		OnGizmoTransformChanged_Handler(Identifier, NewTransform);
 	});
+	GizmoProxy->OnBeginTransformEdit.AddWeakLambda(this, [this, Identifier](UTransformProxy* Proxy) 
+	{
+		OnGizmoTransformStateChange_Handler(Identifier, Proxy->GetTransform(), EScriptableToolGizmoStateChangeType::BeginTransform);
+	});
+	GizmoProxy->OnEndTransformEdit.AddWeakLambda(this, [this, Identifier](UTransformProxy* Proxy)
+	{
+		OnGizmoTransformStateChange_Handler(Identifier, Proxy->GetTransform(), EScriptableToolGizmoStateChangeType::EndTransform);
+	});
+	GizmoProxy->OnTransformChangedUndoRedo.AddWeakLambda(this, [this, Identifier](UTransformProxy*, FTransform NewTransform)
+	{
+		OnGizmoTransformStateChange_Handler(Identifier, NewTransform, EScriptableToolGizmoStateChangeType::UndoRedo);
+	});
+	
 
 	if (GizmoOptions.bAllowNegativeScaling == false)
 	{
@@ -1020,6 +1033,11 @@ void UScriptableInteractiveTool::OnGizmoTransformChanged_Handler(FString GizmoId
 	OnGizmoTransformChanged(GizmoIdentifier, NewTransform);
 }
 
+void UScriptableInteractiveTool::OnGizmoTransformStateChange_Handler(FString GizmoIdentifier, FTransform CurrentTransform, EScriptableToolGizmoStateChangeType ChangeType)
+{
+	// forward to BlueprintImplementableEvent
+	OnGizmoTransformStateChange(GizmoIdentifier, CurrentTransform, ChangeType);
+}
 
 
 
diff --git a/Engine/Plugins/Experimental/ScriptableToolsFramework/Source/ScriptableToolsFramework/Public/ScriptableInteractiveTool.h b/Engine/Plugins/Experimental/ScriptableToolsFramework/Source/ScriptableToolsFramework/Public/ScriptableInteractiveTool.h
--- a/Engine/Plugins/Experimental/ScriptableToolsFramework/Source/ScriptableToolsFramework/Public/ScriptableInteractiveTool.h
+++ b/Engine/Plugins/Experimental/ScriptableToolsFramework/Source/ScriptableToolsFramework/Public/ScriptableInteractiveTool.h
@@ -11,6 +11,7 @@ class UWorld;
 class UCombinedTransformGizmo;
 class UTransformProxy;
 class UScriptableInteractiveTool;
+class UBaseScriptableToolBuilder;
 class FToolDataVisualizer;
 class FCanvas;
 class FSceneView;
@@ -65,6 +66,14 @@ enum class EScriptableToolGizmoCoordinateSystem : uint8
 	FromViewportSettings = 2
 };
 
+UENUM(BlueprintType)
+enum class EScriptableToolGizmoStateChangeType : uint8
+{
+	BeginTransform = 0,
+	EndTransform = 1,
+	UndoRedo = 2
+};
+
 UENUM(BlueprintType, meta = (Bitflags, UseEnumValuesAsMaskValuesInEditor = "true"))
 enum class EScriptableToolGizmoTranslation : uint8
 {
@@ -309,7 +318,6 @@ public:
 	UPROPERTY(BlueprintReadWrite, EditAnywhere, Category = "Scriptable Tool Settings", meta=(DisplayName="Shutdown Type"))
 	EScriptableToolShutdownType ToolShutdownType = EScriptableToolShutdownType::Complete;
 
-
 	/**
 	 * Implement OnScriptSetup to do initial setup/configuration of the Tool, such as adding
 	 * Property Sets, creating Gizmos, etc
@@ -376,6 +384,11 @@ protected:
 	virtual void PostInitProperties() override;
 
 
+	// return instance of custom tool builder. Should only be called on CDO.
+	virtual UBaseScriptableToolBuilder* GetNewCustomToolBuilderInstance(UObject* Outer) { return nullptr; }
+	friend class UScriptableToolSet;
+
+
 protected:
 	UPROPERTY(Transient, DuplicateTransient, NonTransactional, SkipSerialization)
 	TWeakObjectPtr<UWorld> TargetWorld;
@@ -691,12 +704,22 @@ public:
 	UFUNCTION(BlueprintImplementableEvent, Category = "ScriptableTool|Events")
 	void OnGizmoTransformChanged(const FString& GizmoIdentifier, FTransform NewTransform);
 
+	/**
+	 * The OnGizmoTransformStateChange event fires whenever the user start/ends a Gizmo transform, or when an Undo/Redo event occurs.
+	 * Note that when Undo/Redo occurs, OnGizmoTransformChanged will also fire.
+	 * The GizmoIdentifier can be used to disambiguate multiple active Gizmos.
+	 */
+	UFUNCTION(BlueprintImplementableEvent, Category = "ScriptableTool|Events")
+	void OnGizmoTransformStateChange(const FString& GizmoIdentifier, FTransform CurrentTransform, EScriptableToolGizmoStateChangeType ChangeType);
+
+
 protected:
 	// trying to avoid making a UStruct for this internal stuff
 	UPROPERTY(Transient, DuplicateTransient, NonTransactional, SkipSerialization)
 	TMap<FString, TObjectPtr<UCombinedTransformGizmo>> Gizmos;
 
 	virtual void OnGizmoTransformChanged_Handler(FString GizmoIdentifier, FTransform NewTransform);
+	virtual void OnGizmoTransformStateChange_Handler(FString GizmoIdentifier, FTransform CurrentTransform, EScriptableToolGizmoStateChangeType ChangeType);
 
 
 
