ScriptableToolsFramework: Added support for a new Modular Behavior base tool, designed to support the various input behaviors available in the Interactive Tools Framework for blueprintable Scriptable Tools.

Deriving from ScriptableModularBehaviorTool, or the Editor mode EditorScriptableModularBehaviorTool, provides access to a set of blueprint routines for adding mouse input behaviors to the tool. These functions each require a specific set of event callbacks to handle the variety of input event states the behaviors can generate. Users can add as many input behaviors as makes sense for their needs, with customizable priorty, mouse button, and modifier exclusivity to separate when different behaviors get called and executed by the tool.

diff --git a/Engine/Plugins/Runtime/ScriptableToolsFramework/Source/EditorScriptableToolsFramework/Public/BaseTools/EditorScriptableModularBehaviorTool.h b/Engine/Plugins/Runtime/ScriptableToolsFramework/Source/EditorScriptableToolsFramework/Public/BaseTools/EditorScriptableModularBehaviorTool.h
new file mode 100644
--- /dev/null
+++ b/Engine/Plugins/Runtime/ScriptableToolsFramework/Source/EditorScriptableToolsFramework/Public/BaseTools/EditorScriptableModularBehaviorTool.h
@@ -0,0 +1,17 @@
+// Copyright Epic Games, Inc. All Rights Reserved.
+
+#pragma once
+
+#include "BaseTools/ScriptableModularBehaviorTool.h"
+#include "EditorScriptableModularBehaviorTool.generated.h"
+
+/**
+ * Editor-Only variant of UScriptableModularBehaviorTool, which gives access to Editor-Only BP functions
+ */
+UCLASS(Transient, Blueprintable)
+class EDITORSCRIPTABLETOOLSFRAMEWORK_API UEditorScriptableModularBehaviorTool : public UScriptableModularBehaviorTool
+{
+	GENERATED_BODY()
+public:
+
+};
\ No newline at end of file
diff --git a/Engine/Plugins/Runtime/ScriptableToolsFramework/Source/ScriptableToolsFramework/Private/BaseTools/ScriptableModularBehaviorTool.cpp b/Engine/Plugins/Runtime/ScriptableToolsFramework/Source/ScriptableToolsFramework/Private/BaseTools/ScriptableModularBehaviorTool.cpp
new file mode 100644
--- /dev/null
+++ b/Engine/Plugins/Runtime/ScriptableToolsFramework/Source/ScriptableToolsFramework/Private/BaseTools/ScriptableModularBehaviorTool.cpp
@@ -0,0 +1,173 @@
+// Copyright Epic Games, Inc. All Rights Reserved.
+
+#include "BaseTools/ScriptableModularBehaviorTool.h"
+
+#include "Behaviors/ScriptableToolSingleClickBehavior.h"
+#include "Behaviors/ScriptableToolClickDragBehavior.h"
+#include "Behaviors/ScriptableToolMouseHoverBehavior.h"
+#include "Behaviors/ScriptableToolMouseWheelBehavior.h"
+#include "Behaviors/ScriptableToolMultiClickSequenceBehavior.h"
+#include "Behaviors/ScriptableToolSingleClickOrDragBehavior.h"
+
+void UScriptableModularBehaviorTool::AddSingleClickBehavior(
+	const FTestIfHitByClickDelegate& TestIfHitByClickDelegate,
+	const FOnHitByClickDelegate& OnHitByClickDelegate,
+	const FMouseBehaviorModiferCheckDelegate& ModifierCheckFunction,
+	int CapturePriority,
+	EScriptableToolMouseButton MouseButton,
+	bool bHitTestOnRelease
+)
+{
+	TObjectPtr<UScriptableToolSingleClickBehavior> BehaviorContainer = NewObject<UScriptableToolSingleClickBehavior>();
+	BehaviorContainer->Init(this, ModifierCheckFunction, TestIfHitByClickDelegate, OnHitByClickDelegate, MouseButton, bHitTestOnRelease);
+	BehaviorContainer->SetDefaultPriority(FInputCapturePriority(CapturePriority));
+
+	SingleClickBehaviors.Add(BehaviorContainer);
+}
+
+
+void UScriptableModularBehaviorTool::AddClickDragBehavior(
+	const FTestCanBeginClickDragSequenceDelegate& TestCanBeginClickDragSequenceFuncIn,
+	const FOnClickPressDelegate& OnClickPressFuncIn,
+	const FOnClickDragDelegate& OnClickDragFuncIn,
+	const FOnClickReleaseDelegate& OnClickReleaseFuncIn,
+	const FOnTerminateDragSequenceDelegate& OnTerminateDragSequenceFuncIn,
+	const FMouseBehaviorModiferCheckDelegate& ModifierCheckFuncIn,
+	int CapturePriority,
+	EScriptableToolMouseButton MouseButtonIn,
+	bool bUpdateModifiersDuringDrag
+)
+{
+	TObjectPtr<UScriptableToolClickDragBehavior> BehaviorContainer = NewObject<UScriptableToolClickDragBehavior>();
+	BehaviorContainer->Init(this, ModifierCheckFuncIn, TestCanBeginClickDragSequenceFuncIn, OnClickPressFuncIn, OnClickDragFuncIn, OnClickReleaseFuncIn, OnTerminateDragSequenceFuncIn, MouseButtonIn, bUpdateModifiersDuringDrag);
+	BehaviorContainer->SetDefaultPriority(FInputCapturePriority(CapturePriority));
+
+	ClickDragBehaviors.Add(BehaviorContainer);
+}
+
+
+void UScriptableModularBehaviorTool::AddSingleClickOrDragBehavior(
+	FTestIfHitByClickDelegate TestIfHitByClickFuncIn,
+	FOnHitByClickDelegate OnHitByClickFuncIn,
+	FTestCanBeginClickDragSequenceDelegate TestCanBeginClickDragSequenceFuncIn,
+	FOnClickPressDelegate OnClickPressFuncIn,
+	FOnClickDragDelegate OnClickDragFuncIn,
+	FOnClickReleaseDelegate OnClickReleaseFuncIn,
+	FOnTerminateDragSequenceDelegate OnTerminateDragSequenceFuncIn,
+	FMouseBehaviorModiferCheckDelegate ModifierCheckFuncIn,
+	int CapturePriority,
+	EScriptableToolMouseButton MouseButtonIn,
+	bool bBeginDragIfClickTargetNotHit,
+	float ClickDistanceThreshold
+)
+{
+	TObjectPtr<UScriptableToolSingleClickOrDragBehavior> BehaviorContainer = NewObject<UScriptableToolSingleClickOrDragBehavior>();
+	BehaviorContainer->Init(this, ModifierCheckFuncIn, TestIfHitByClickFuncIn, OnHitByClickFuncIn,
+		                    TestCanBeginClickDragSequenceFuncIn, OnClickPressFuncIn, OnClickDragFuncIn,
+		                    OnClickReleaseFuncIn, OnTerminateDragSequenceFuncIn, MouseButtonIn, 
+		                    bBeginDragIfClickTargetNotHit, ClickDistanceThreshold);
+	BehaviorContainer->SetDefaultPriority(FInputCapturePriority(CapturePriority));
+
+	SingleClickOrDragBehaviors.Add(BehaviorContainer);
+}
+
+
+
+void UScriptableModularBehaviorTool::AddMouseWheelBehavior(
+	FTestShouldRespondToMouseWheelDelegate TestShouldRespondToMouseWheelFuncIn,
+	FOnMouseWheelScrollUpDelegate OnMouseWheelScrollUpFuncIn,
+	FOnMouseWheelScrollDownDelegate OnMouseWheelScrollDownFuncIn,
+	FMouseBehaviorModiferCheckDelegate ModifierCheckFuncIn,
+	int CapturePriority
+)
+{
+	TObjectPtr<UScriptableToolMouseWheelBehavior> BehaviorContainer = NewObject<UScriptableToolMouseWheelBehavior>();
+	BehaviorContainer->Init(this, ModifierCheckFuncIn, TestShouldRespondToMouseWheelFuncIn, OnMouseWheelScrollUpFuncIn, OnMouseWheelScrollDownFuncIn);
+	BehaviorContainer->SetDefaultPriority(FInputCapturePriority(CapturePriority));
+
+	MouseWheelBehaviors.Add(BehaviorContainer);
+}
+
+
+
+void UScriptableModularBehaviorTool::AddMultiClickSequenceBehavior(
+	FOnBeginSequencePreviewDelegate OnBeginSequencePreviewFuncIn,
+	FCanBeginClickSequenceDelegate CanBeginClickSequenceFuncIn,
+	FOnBeginClickSequenceDelegate OnBeginClickSequenceFuncIn,
+	FOnNextSequencePreviewDelegate OnNextSequencePreviewFuncIn,
+	FOnNextSequenceClickDelegate OnNextSequenceClickFuncIn,
+	FOnTerminateClickSequenceDelegate OnTerminateClickSequenceFuncIn,
+	FRequestAbortClickSequenceDelegate RequestAbortClickSequenceFuncIn,
+	FMouseBehaviorModiferCheckDelegate ModifierCheckFuncIn,
+	int CapturePriority,
+	EScriptableToolMouseButton MouseButtonIn
+)
+{
+	TObjectPtr<UScriptableToolClickSequenceBehavior> BehaviorContainer = NewObject<UScriptableToolClickSequenceBehavior>();
+	BehaviorContainer->Init(this, ModifierCheckFuncIn, OnBeginSequencePreviewFuncIn, CanBeginClickSequenceFuncIn, OnBeginClickSequenceFuncIn, OnNextSequencePreviewFuncIn,
+		                    OnNextSequenceClickFuncIn, OnTerminateClickSequenceFuncIn, RequestAbortClickSequenceFuncIn, MouseButtonIn);
+	BehaviorContainer->SetDefaultPriority(FInputCapturePriority(CapturePriority));
+
+	MultiClickSequenceBehaviors.Add(BehaviorContainer);
+}
+
+
+
+
+void UScriptableModularBehaviorTool::AddMouseHoverBehavior(
+	FBeginHoverSequenceHitTestDelegate BeginHoverSequenceHitTestFuncIn,
+	FOnBeginHoverDelegate OnBeginHoverFuncIn,
+	FOnUpdateHoverDelegate OnUpdateHoverFuncIn,
+	FOnEndHoverDelegate OnEndHoverFuncIn,
+	int CapturePriority
+)
+{
+	TObjectPtr<UScriptableToolMouseHoverBehavior> BehaviorContainer = NewObject<UScriptableToolMouseHoverBehavior>();
+	BehaviorContainer->Init(this, BeginHoverSequenceHitTestFuncIn, OnBeginHoverFuncIn, OnUpdateHoverFuncIn, OnEndHoverFuncIn);
+	BehaviorContainer->SetDefaultPriority(FInputCapturePriority(CapturePriority));
+
+	MouseHoverBehaviors.Add(BehaviorContainer);
+}
+
+
+void UScriptableModularBehaviorTool::OnUpdateModifierState(int ModifierID, bool bIsOn)
+{
+	if (ModifierID == 1)
+	{
+		bShiftModifier = bIsOn;
+	}
+	if (ModifierID == 2)
+	{
+		bCtrlModifier = bIsOn;
+	}
+	if (ModifierID == 3)
+	{
+		bAltModifier = bIsOn;
+	}
+}
+
+bool UScriptableModularBehaviorTool::IsShiftDown() const
+{
+	return bShiftModifier;
+}
+
+bool UScriptableModularBehaviorTool::IsCtrlDown() const
+{
+	return bCtrlModifier;
+}
+
+bool UScriptableModularBehaviorTool::IsAltDown() const
+{
+	return bAltModifier;
+}
+
+FScriptableToolModifierStates UScriptableModularBehaviorTool::GetActiveModifiers()
+{
+	FScriptableToolModifierStates Modifiers;
+	Modifiers.bShiftDown = bShiftModifier;
+	Modifiers.bCtrlDown = bCtrlModifier;
+	Modifiers.bAltDown = bAltModifier;
+	return Modifiers;
+}
+
+
diff --git a/Engine/Plugins/Runtime/ScriptableToolsFramework/Source/ScriptableToolsFramework/Private/Behaviors/ScriptableToolBehavior.cpp b/Engine/Plugins/Runtime/ScriptableToolsFramework/Source/ScriptableToolsFramework/Private/Behaviors/ScriptableToolBehavior.cpp
new file mode 100644
--- /dev/null
+++ b/Engine/Plugins/Runtime/ScriptableToolsFramework/Source/ScriptableToolsFramework/Private/Behaviors/ScriptableToolBehavior.cpp
@@ -0,0 +1,8 @@
+#include "ScriptableToolBehavior.h"
+
+void UScriptableToolBehavior::SetDefaultPriority(const FInputCapturePriority& Priority)
+{
+	GetWrappedBehavior()->SetDefaultPriority(Priority);
+}
+
+
diff --git a/Engine/Plugins/Runtime/ScriptableToolsFramework/Source/ScriptableToolsFramework/Private/Behaviors/ScriptableToolBehavior.h b/Engine/Plugins/Runtime/ScriptableToolsFramework/Source/ScriptableToolsFramework/Private/Behaviors/ScriptableToolBehavior.h
new file mode 100644
--- /dev/null
+++ b/Engine/Plugins/Runtime/ScriptableToolsFramework/Source/ScriptableToolsFramework/Private/Behaviors/ScriptableToolBehavior.h
@@ -0,0 +1,22 @@
+// Copyright Epic Games, Inc. All Rights Reserved.
+
+#pragma once
+
+#include "InputBehavior.h"
+
+#include "ScriptableToolBehavior.generated.h"
+
+UCLASS(Abstract)
+class UScriptableToolBehavior : public UObject
+{
+	GENERATED_BODY()
+
+public:
+
+	UScriptableToolBehavior() {};
+
+	void SetDefaultPriority(const FInputCapturePriority& Priority);
+
+	virtual UInputBehavior* GetWrappedBehavior() PURE_VIRTUAL(UScriptableToolBehavior::GetWrappedBehavior, return nullptr;);
+
+};
\ No newline at end of file
diff --git a/Engine/Plugins/Runtime/ScriptableToolsFramework/Source/ScriptableToolsFramework/Private/Behaviors/ScriptableToolClickDragBehavior.cpp b/Engine/Plugins/Runtime/ScriptableToolsFramework/Source/ScriptableToolsFramework/Private/Behaviors/ScriptableToolClickDragBehavior.cpp
new file mode 100644
--- /dev/null
+++ b/Engine/Plugins/Runtime/ScriptableToolsFramework/Source/ScriptableToolsFramework/Private/Behaviors/ScriptableToolClickDragBehavior.cpp
@@ -0,0 +1,107 @@
+// Copyright Epic Games, Inc. All Rights Reserved.
+
+#include "Behaviors/ScriptableToolClickDragBehavior.h"
+
+#include "BaseTools/ScriptableModularBehaviorTool.h"
+#include "BaseBehaviors/ClickDragBehavior.h"
+
+
+void UScriptableToolClickDragBehavior::Init(TObjectPtr<UScriptableModularBehaviorTool> BehaviorHostIn,
+	FMouseBehaviorModiferCheckDelegate ModifierCheckFuncIn,
+	FTestCanBeginClickDragSequenceDelegate TestCanBeginClickDragSequenceFuncIn,
+	FOnClickPressDelegate OnClickPressFuncIn,
+	FOnClickDragDelegate OnClickDragFuncIn,
+	FOnClickReleaseDelegate OnClickReleaseFuncIn,
+	FOnTerminateDragSequenceDelegate OnTerminateDragSequenceFuncIn,
+	EScriptableToolMouseButton MouseButtonIn,
+	bool bUpdateModifiersDuringDragIn)
+{
+	BehaviorHost = BehaviorHostIn;
+	Behavior = NewObject<UClickDragInputBehavior>();
+	ModifierCheckFunc = ModifierCheckFuncIn;
+	TestCanBeginClickDragSequenceFunc = TestCanBeginClickDragSequenceFuncIn;
+	OnClickPressFunc = OnClickPressFuncIn;
+	OnClickDragFunc = OnClickDragFuncIn;
+	OnClickReleaseFunc = OnClickReleaseFuncIn;
+	OnTerminateDragSequenceFunc = OnTerminateDragSequenceFuncIn;
+	MouseButton = MouseButtonIn;
+
+	Behavior->Initialize(this);
+	Behavior->ModifierCheckFunc = [this](const FInputDeviceState& InputDeviceState) {
+		bool bResult = true;
+		if (ModifierCheckFunc.IsBound())
+		{
+			bResult = ModifierCheckFunc.Execute(InputDeviceState);
+		}
+		return bResult;
+	};
+
+	Behavior->bUpdateModifiersDuringDrag = bUpdateModifiersDuringDragIn;
+
+	BehaviorHost->AddInputBehavior(Behavior);
+
+	Behavior->Modifiers.RegisterModifier(1, FInputDeviceState::IsShiftKeyDown);
+	Behavior->Modifiers.RegisterModifier(2, FInputDeviceState::IsCtrlKeyDown);
+	Behavior->Modifiers.RegisterModifier(3, FInputDeviceState::IsAltKeyDown);
+
+	switch (MouseButton)
+	{
+	case EScriptableToolMouseButton::LeftButton:
+	{
+		Behavior->SetUseLeftMouseButton();
+		break;
+	}
+	case EScriptableToolMouseButton::RightButton:
+	{
+		Behavior->SetUseRightMouseButton();
+		break;
+	}
+	case EScriptableToolMouseButton::MiddleButton:
+	{
+		Behavior->SetUseMiddleMouseButton();
+		break;
+	}
+	default:
+		ensure(false);
+	}
+}
+
+FInputRayHit UScriptableToolClickDragBehavior::CanBeginClickDragSequence(const FInputDeviceRay& PressPos)
+{
+	FInputRayHit RayHit;
+	if (TestCanBeginClickDragSequenceFunc.IsBound())
+	{
+		RayHit = TestCanBeginClickDragSequenceFunc.Execute(PressPos, BehaviorHost->GetActiveModifiers(), MouseButton);
+	}
+	return RayHit;
+}
+
+UInputBehavior* UScriptableToolClickDragBehavior::GetWrappedBehavior()
+{
+	return Behavior;
+}
+
+void UScriptableToolClickDragBehavior::OnClickPress(const FInputDeviceRay& PressPos)
+{
+	OnClickPressFunc.ExecuteIfBound(PressPos, BehaviorHost->GetActiveModifiers(), MouseButton);
+}
+
+void UScriptableToolClickDragBehavior::OnClickDrag(const FInputDeviceRay& DragPos)
+{
+	OnClickDragFunc.ExecuteIfBound(DragPos, BehaviorHost->GetActiveModifiers(), MouseButton);
+}
+
+void UScriptableToolClickDragBehavior::OnClickRelease(const FInputDeviceRay& ReleasePos)
+{
+	OnClickReleaseFunc.ExecuteIfBound(ReleasePos, BehaviorHost->GetActiveModifiers(), MouseButton);
+}
+
+void UScriptableToolClickDragBehavior::OnTerminateDragSequence()
+{
+	OnTerminateDragSequenceFunc.ExecuteIfBound(BehaviorHost->GetActiveModifiers(), MouseButton);
+}
+
+void UScriptableToolClickDragBehavior::OnUpdateModifierState(int ModifierID, bool bIsOn)
+{
+	BehaviorHost->OnUpdateModifierState(ModifierID, bIsOn);
+}
\ No newline at end of file
diff --git a/Engine/Plugins/Runtime/ScriptableToolsFramework/Source/ScriptableToolsFramework/Private/Behaviors/ScriptableToolClickDragBehavior.h b/Engine/Plugins/Runtime/ScriptableToolsFramework/Source/ScriptableToolsFramework/Private/Behaviors/ScriptableToolClickDragBehavior.h
new file mode 100644
--- /dev/null
+++ b/Engine/Plugins/Runtime/ScriptableToolsFramework/Source/ScriptableToolsFramework/Private/Behaviors/ScriptableToolClickDragBehavior.h
@@ -0,0 +1,70 @@
+// Copyright Epic Games, Inc. All Rights Reserved.
+
+#pragma once
+
+#include "Behaviors/ScriptableToolBehaviorDelegates.h"
+#include "ScriptableToolBehavior.h"
+#include "BaseBehaviors/BehaviorTargetInterfaces.h"
+
+#include "ScriptableToolClickDragBehavior.generated.h"
+
+class UScriptableModularBehaviorTool;
+class UClickDragInputBehavior;
+
+
+UCLASS()
+class UScriptableToolClickDragBehavior : public UScriptableToolBehavior, public IClickDragBehaviorTarget
+{
+	GENERATED_BODY()
+
+public:
+	UScriptableToolClickDragBehavior() {};
+
+	void Init(TObjectPtr<UScriptableModularBehaviorTool> BehaviorHostIn,
+		FMouseBehaviorModiferCheckDelegate ModifierCheckFuncIn,
+		FTestCanBeginClickDragSequenceDelegate TestCanBeginClickDragSequenceFuncIn,
+		FOnClickPressDelegate OnClickPressFuncIn,
+		FOnClickDragDelegate OnClickDragFuncIn,
+		FOnClickReleaseDelegate OnClickReleaseFuncIn,
+		FOnTerminateDragSequenceDelegate OnTerminateDragSequenceFuncIn,
+		EScriptableToolMouseButton MouseButtonIn,
+		bool bUpdateModifiersDuringDragIn);
+
+	virtual UInputBehavior* GetWrappedBehavior() override;
+
+	virtual FInputRayHit CanBeginClickDragSequence(const FInputDeviceRay& PressPos) override;
+	virtual void OnClickPress(const FInputDeviceRay& PressPos) override;
+	virtual void OnClickDrag(const FInputDeviceRay& DragPos) override;
+	virtual void OnClickRelease(const FInputDeviceRay& ReleasePos) override;
+	virtual void OnTerminateDragSequence() override;
+	virtual void OnUpdateModifierState(int ModifierID, bool bIsOn) override;
+
+private:
+
+	UPROPERTY()
+	TObjectPtr<UScriptableModularBehaviorTool> BehaviorHost;
+
+	UPROPERTY()
+	TObjectPtr<UClickDragInputBehavior> Behavior;
+
+	UPROPERTY()
+	FMouseBehaviorModiferCheckDelegate ModifierCheckFunc;
+
+	UPROPERTY()
+	FTestCanBeginClickDragSequenceDelegate	TestCanBeginClickDragSequenceFunc;
+
+	UPROPERTY()
+	FOnClickPressDelegate OnClickPressFunc;
+
+	UPROPERTY()
+	FOnClickDragDelegate OnClickDragFunc;
+
+	UPROPERTY()
+	FOnClickReleaseDelegate OnClickReleaseFunc;
+
+	UPROPERTY()
+	FOnTerminateDragSequenceDelegate OnTerminateDragSequenceFunc;
+
+	UPROPERTY()
+	EScriptableToolMouseButton MouseButton;
+};
\ No newline at end of file
diff --git a/Engine/Plugins/Runtime/ScriptableToolsFramework/Source/ScriptableToolsFramework/Private/Behaviors/ScriptableToolMouseHoverBehavior.cpp b/Engine/Plugins/Runtime/ScriptableToolsFramework/Source/ScriptableToolsFramework/Private/Behaviors/ScriptableToolMouseHoverBehavior.cpp
new file mode 100644
--- /dev/null
+++ b/Engine/Plugins/Runtime/ScriptableToolsFramework/Source/ScriptableToolsFramework/Private/Behaviors/ScriptableToolMouseHoverBehavior.cpp
@@ -0,0 +1,70 @@
+// Copyright Epic Games, Inc. All Rights Reserved.
+
+#include "Behaviors/ScriptableToolMouseHoverBehavior.h"
+
+#include "BaseTools/ScriptableModularBehaviorTool.h"
+#include "BaseBehaviors/MouseHoverBehavior.h"
+
+
+void UScriptableToolMouseHoverBehavior::Init(TObjectPtr<UScriptableModularBehaviorTool> BehaviorHostIn,
+	FBeginHoverSequenceHitTestDelegate BeginHoverSequenceHitTestFuncIn,
+	FOnBeginHoverDelegate OnBeginHoverFuncIn,
+	FOnUpdateHoverDelegate OnUpdateHoverFuncIn,
+	FOnEndHoverDelegate OnEndHoverFuncIn)
+{
+	BehaviorHost = BehaviorHostIn;
+	Behavior = NewObject<UMouseHoverBehavior>();
+	BeginHoverSequenceHitTestFunc = BeginHoverSequenceHitTestFuncIn;
+	OnBeginHoverFunc = OnBeginHoverFuncIn;
+	OnUpdateHoverFunc = OnUpdateHoverFuncIn;
+	OnEndHoverFunc = OnEndHoverFuncIn;
+
+	Behavior->Initialize(this);
+
+
+	BehaviorHost->AddInputBehavior(Behavior);
+
+	Behavior->Modifiers.RegisterModifier(1, FInputDeviceState::IsShiftKeyDown);
+	Behavior->Modifiers.RegisterModifier(2, FInputDeviceState::IsCtrlKeyDown);
+	Behavior->Modifiers.RegisterModifier(3, FInputDeviceState::IsAltKeyDown);
+}
+
+FInputRayHit UScriptableToolMouseHoverBehavior::BeginHoverSequenceHitTest(const FInputDeviceRay& PressPos)
+{
+	FInputRayHit RayHit;
+	if (BeginHoverSequenceHitTestFunc.IsBound())
+	{
+		RayHit = BeginHoverSequenceHitTestFunc.Execute(PressPos, BehaviorHost->GetActiveModifiers());
+	}
+	return RayHit;
+}
+
+UInputBehavior* UScriptableToolMouseHoverBehavior::GetWrappedBehavior()
+{
+	return Behavior;
+}
+
+void UScriptableToolMouseHoverBehavior::OnBeginHover(const FInputDeviceRay& DevicePos)
+{
+	OnBeginHoverFunc.ExecuteIfBound(DevicePos, BehaviorHost->GetActiveModifiers());
+}
+
+bool UScriptableToolMouseHoverBehavior::OnUpdateHover(const FInputDeviceRay& DevicePos)
+{
+	bool bShouldContinue = false;
+	if (OnUpdateHoverFunc.IsBound())
+	{
+		bShouldContinue = OnUpdateHoverFunc.Execute(DevicePos, BehaviorHost->GetActiveModifiers());
+	}
+	return bShouldContinue;
+}
+
+void UScriptableToolMouseHoverBehavior::OnEndHover()
+{
+	OnEndHoverFunc.ExecuteIfBound();
+}
+
+void UScriptableToolMouseHoverBehavior::OnUpdateModifierState(int ModifierID, bool bIsOn)
+{
+	BehaviorHost->OnUpdateModifierState(ModifierID, bIsOn);
+}
\ No newline at end of file
diff --git a/Engine/Plugins/Runtime/ScriptableToolsFramework/Source/ScriptableToolsFramework/Private/Behaviors/ScriptableToolMouseHoverBehavior.h b/Engine/Plugins/Runtime/ScriptableToolsFramework/Source/ScriptableToolsFramework/Private/Behaviors/ScriptableToolMouseHoverBehavior.h
new file mode 100644
--- /dev/null
+++ b/Engine/Plugins/Runtime/ScriptableToolsFramework/Source/ScriptableToolsFramework/Private/Behaviors/ScriptableToolMouseHoverBehavior.h
@@ -0,0 +1,57 @@
+// Copyright Epic Games, Inc. All Rights Reserved.
+
+#pragma once
+
+#include "Behaviors/ScriptableToolBehaviorDelegates.h"
+#include "ScriptableToolBehavior.h"
+#include "BaseBehaviors/BehaviorTargetInterfaces.h"
+
+#include "ScriptableToolMouseHoverBehavior.generated.h"
+
+class UScriptableModularBehaviorTool;
+class UMouseHoverBehavior;
+
+
+UCLASS()
+class UScriptableToolMouseHoverBehavior : public UScriptableToolBehavior, public IHoverBehaviorTarget
+{
+	GENERATED_BODY()
+
+public:
+	UScriptableToolMouseHoverBehavior() {};
+
+	void Init(TObjectPtr<UScriptableModularBehaviorTool> BehaviorHostIn,
+		FBeginHoverSequenceHitTestDelegate BeginHoverSequenceHitTestFuncIn,
+		FOnBeginHoverDelegate OnBeginHoverFuncIn,
+		FOnUpdateHoverDelegate OnUpdateHoverFuncIn,
+		FOnEndHoverDelegate OnEndHoverFuncIn);
+
+	virtual UInputBehavior* GetWrappedBehavior() override;
+
+	virtual FInputRayHit BeginHoverSequenceHitTest(const FInputDeviceRay& PressPos) override;
+	virtual void OnBeginHover(const FInputDeviceRay& DevicePos) override;
+	virtual bool OnUpdateHover(const FInputDeviceRay& DevicePos) override;
+	virtual void OnEndHover() override;
+	virtual void OnUpdateModifierState(int ModifierID, bool bIsOn) override;
+
+private:
+
+	UPROPERTY()
+	TObjectPtr<UScriptableModularBehaviorTool> BehaviorHost;
+
+	UPROPERTY()
+	TObjectPtr<UMouseHoverBehavior> Behavior;
+
+	UPROPERTY()
+	FBeginHoverSequenceHitTestDelegate BeginHoverSequenceHitTestFunc;
+
+	UPROPERTY()
+	FOnBeginHoverDelegate OnBeginHoverFunc;
+
+	UPROPERTY()
+	FOnUpdateHoverDelegate OnUpdateHoverFunc;
+
+	UPROPERTY()
+	FOnEndHoverDelegate OnEndHoverFunc;
+
+};
\ No newline at end of file
diff --git a/Engine/Plugins/Runtime/ScriptableToolsFramework/Source/ScriptableToolsFramework/Private/Behaviors/ScriptableToolMouseWheelBehavior.cpp b/Engine/Plugins/Runtime/ScriptableToolsFramework/Source/ScriptableToolsFramework/Private/Behaviors/ScriptableToolMouseWheelBehavior.cpp
new file mode 100644
--- /dev/null
+++ b/Engine/Plugins/Runtime/ScriptableToolsFramework/Source/ScriptableToolsFramework/Private/Behaviors/ScriptableToolMouseWheelBehavior.cpp
@@ -0,0 +1,69 @@
+// Copyright Epic Games, Inc. All Rights Reserved.
+
+#include "Behaviors/ScriptableToolMouseWheelBehavior.h"
+
+#include "BaseTools/ScriptableModularBehaviorTool.h"
+#include "BaseBehaviors/MouseWheelBehavior.h"
+
+
+void UScriptableToolMouseWheelBehavior::Init(TObjectPtr<UScriptableModularBehaviorTool> BehaviorHostIn,
+	FMouseBehaviorModiferCheckDelegate ModifierCheckFuncIn,
+	FTestShouldRespondToMouseWheelDelegate TestShouldRespondToMouseWheelFuncIn,
+	FOnMouseWheelScrollUpDelegate OnMouseWheelScrollUpFuncIn,
+	FOnMouseWheelScrollDownDelegate OnMouseWheelScrollDownFuncIn)
+{
+	BehaviorHost = BehaviorHostIn;
+	Behavior = NewObject<UMouseWheelInputBehavior>();
+	ModifierCheckFunc = ModifierCheckFuncIn;
+	TestShouldRespondToMouseWheelFunc = TestShouldRespondToMouseWheelFuncIn;
+	OnMouseWheelScrollUpFunc = OnMouseWheelScrollUpFuncIn;
+	OnMouseWheelScrollDownFunc = OnMouseWheelScrollDownFuncIn;
+
+	Behavior->Initialize(this);
+	Behavior->ModifierCheckFunc = [this](const FInputDeviceState& InputDeviceState) {
+		bool bResult = true;
+		if (ModifierCheckFunc.IsBound())
+		{
+			bResult = ModifierCheckFunc.Execute(InputDeviceState);
+		}
+		return bResult;
+	};
+
+
+	BehaviorHost->AddInputBehavior(Behavior);
+
+	Behavior->Modifiers.RegisterModifier(1, FInputDeviceState::IsShiftKeyDown);
+	Behavior->Modifiers.RegisterModifier(2, FInputDeviceState::IsCtrlKeyDown);
+	Behavior->Modifiers.RegisterModifier(3, FInputDeviceState::IsAltKeyDown);
+
+}
+
+UInputBehavior* UScriptableToolMouseWheelBehavior::GetWrappedBehavior()
+{
+	return Behavior;
+}
+
+FInputRayHit UScriptableToolMouseWheelBehavior::ShouldRespondToMouseWheel(const FInputDeviceRay& CurrentPos)
+{
+	FInputRayHit RayHit;
+	if (TestShouldRespondToMouseWheelFunc.IsBound())
+	{
+		RayHit = TestShouldRespondToMouseWheelFunc.Execute(CurrentPos, BehaviorHost->GetActiveModifiers());
+	}
+	return RayHit;
+}
+
+void UScriptableToolMouseWheelBehavior::OnMouseWheelScrollUp(const FInputDeviceRay& CurrentPos)
+{
+	OnMouseWheelScrollUpFunc.ExecuteIfBound(CurrentPos, BehaviorHost->GetActiveModifiers());
+}
+
+void UScriptableToolMouseWheelBehavior::OnMouseWheelScrollDown(const FInputDeviceRay& CurrentPos)
+{
+	OnMouseWheelScrollUpFunc.ExecuteIfBound(CurrentPos, BehaviorHost->GetActiveModifiers());
+}
+
+void UScriptableToolMouseWheelBehavior::OnUpdateModifierState(int ModifierID, bool bIsOn)
+{
+	BehaviorHost->OnUpdateModifierState(ModifierID, bIsOn);
+}
diff --git a/Engine/Plugins/Runtime/ScriptableToolsFramework/Source/ScriptableToolsFramework/Private/Behaviors/ScriptableToolMouseWheelBehavior.h b/Engine/Plugins/Runtime/ScriptableToolsFramework/Source/ScriptableToolsFramework/Private/Behaviors/ScriptableToolMouseWheelBehavior.h
new file mode 100644
--- /dev/null
+++ b/Engine/Plugins/Runtime/ScriptableToolsFramework/Source/ScriptableToolsFramework/Private/Behaviors/ScriptableToolMouseWheelBehavior.h
@@ -0,0 +1,56 @@
+// Copyright Epic Games, Inc. All Rights Reserved.
+
+#pragma once
+
+#include "Behaviors/ScriptableToolBehaviorDelegates.h"
+#include "ScriptableToolBehavior.h"
+#include "BaseBehaviors/BehaviorTargetInterfaces.h"
+
+#include "ScriptableToolMouseWheelBehavior.generated.h"
+
+class UScriptableModularBehaviorTool;
+class UMouseWheelInputBehavior;
+
+
+UCLASS()
+class UScriptableToolMouseWheelBehavior : public UScriptableToolBehavior, public IMouseWheelBehaviorTarget
+{
+	GENERATED_BODY()
+
+public:
+	UScriptableToolMouseWheelBehavior() {};
+
+	void Init(TObjectPtr<UScriptableModularBehaviorTool> BehaviorHostIn,
+		FMouseBehaviorModiferCheckDelegate ModifierCheckFuncIn,
+		FTestShouldRespondToMouseWheelDelegate TestShouldRespondToMouseWheelFuncIn,
+		FOnMouseWheelScrollUpDelegate OnMouseWheelScrollUpFuncIn,
+		FOnMouseWheelScrollDownDelegate OnMouseWheelScrollDownFuncIn);
+
+	virtual UInputBehavior* GetWrappedBehavior() override;
+
+	virtual FInputRayHit ShouldRespondToMouseWheel(const FInputDeviceRay& CurrentPos) override;
+	virtual void OnMouseWheelScrollUp(const FInputDeviceRay& CurrentPos) override;
+	virtual void OnMouseWheelScrollDown(const FInputDeviceRay& CurrentPos) override;
+	virtual void OnUpdateModifierState(int ModifierID, bool bIsOn) override;
+
+private:
+
+	UPROPERTY()
+	TObjectPtr<UScriptableModularBehaviorTool> BehaviorHost;
+
+	UPROPERTY()
+	TObjectPtr<UMouseWheelInputBehavior> Behavior;
+
+	UPROPERTY()
+	FMouseBehaviorModiferCheckDelegate ModifierCheckFunc;
+
+	UPROPERTY()
+	FTestShouldRespondToMouseWheelDelegate	TestShouldRespondToMouseWheelFunc;
+
+	UPROPERTY()
+	FOnMouseWheelScrollUpDelegate OnMouseWheelScrollUpFunc;
+
+	UPROPERTY()
+	FOnMouseWheelScrollDownDelegate OnMouseWheelScrollDownFunc;
+
+};
\ No newline at end of file
diff --git a/Engine/Plugins/Runtime/ScriptableToolsFramework/Source/ScriptableToolsFramework/Private/Behaviors/ScriptableToolMultiClickSequenceBehavior.cpp b/Engine/Plugins/Runtime/ScriptableToolsFramework/Source/ScriptableToolsFramework/Private/Behaviors/ScriptableToolMultiClickSequenceBehavior.cpp
new file mode 100644
--- /dev/null
+++ b/Engine/Plugins/Runtime/ScriptableToolsFramework/Source/ScriptableToolsFramework/Private/Behaviors/ScriptableToolMultiClickSequenceBehavior.cpp
@@ -0,0 +1,129 @@
+// Copyright Epic Games, Inc. All Rights Reserved.
+
+#include "Behaviors/ScriptableToolMultiClickSequenceBehavior.h"
+
+#include "BaseTools/ScriptableModularBehaviorTool.h"
+#include "BaseBehaviors/MultiClickSequenceInputBehavior.h"
+
+void UScriptableToolClickSequenceBehavior::Init(TObjectPtr<UScriptableModularBehaviorTool> BehaviorHostIn,
+	FMouseBehaviorModiferCheckDelegate ModifierCheckFuncIn,
+	FOnBeginSequencePreviewDelegate OnBeginSequencePreviewFuncIn,
+	FCanBeginClickSequenceDelegate CanBeginClickSequenceFuncIn,
+	FOnBeginClickSequenceDelegate OnBeginClickSequenceFuncIn,
+	FOnNextSequencePreviewDelegate OnNextSequencePreviewFuncIn,
+	FOnNextSequenceClickDelegate OnNextSequenceClickFuncIn,
+	FOnTerminateClickSequenceDelegate OnTerminateClickSequenceFuncIn,
+	FRequestAbortClickSequenceDelegate RequestAbortClickSequenceFuncIn,
+	EScriptableToolMouseButton MouseButtonIn)
+{
+	BehaviorHost = BehaviorHostIn;
+	Behavior = NewObject<UMultiClickSequenceInputBehavior>();
+	ModifierCheckFunc = ModifierCheckFuncIn;
+	OnBeginSequencePreviewFunc = OnBeginSequencePreviewFuncIn;
+	CanBeginClickSequenceFunc = CanBeginClickSequenceFuncIn;
+	OnBeginClickSequenceFunc = OnBeginClickSequenceFuncIn;
+	OnNextSequencePreviewFunc = OnNextSequencePreviewFuncIn;
+	OnNextSequenceClickFunc = OnNextSequenceClickFuncIn;
+	OnTerminateClickSequenceFunc = OnTerminateClickSequenceFuncIn;
+	RequestAbortClickSequenceFunc = RequestAbortClickSequenceFuncIn;
+	MouseButton = MouseButtonIn;
+
+	Behavior->Initialize(this);
+	Behavior->ModifierCheckFunc = [this](const FInputDeviceState& InputDeviceState) {
+		bool bResult = true;
+		if (ModifierCheckFunc.IsBound())
+		{
+			bResult = ModifierCheckFunc.Execute(InputDeviceState);
+		}
+		return bResult;
+	};
+
+
+	BehaviorHost->AddInputBehavior(Behavior);
+
+	Behavior->Modifiers.RegisterModifier(1, FInputDeviceState::IsShiftKeyDown);
+	Behavior->Modifiers.RegisterModifier(2, FInputDeviceState::IsCtrlKeyDown);
+	Behavior->Modifiers.RegisterModifier(3, FInputDeviceState::IsAltKeyDown);
+
+	switch (MouseButton)
+	{
+	case EScriptableToolMouseButton::LeftButton:
+	{
+		Behavior->SetUseLeftMouseButton();
+		break;
+	}
+	case EScriptableToolMouseButton::RightButton:
+	{
+		Behavior->SetUseRightMouseButton();
+		break;
+	}
+	case EScriptableToolMouseButton::MiddleButton:
+	{
+		Behavior->SetUseMiddleMouseButton();
+		break;
+	}
+	default:
+		ensure(false);
+	}
+}
+
+UInputBehavior* UScriptableToolClickSequenceBehavior::GetWrappedBehavior()
+{
+	return Behavior;
+}
+
+void UScriptableToolClickSequenceBehavior::OnBeginSequencePreview(const FInputDeviceRay& ClickPos)
+{
+	OnBeginSequencePreviewFunc.ExecuteIfBound(ClickPos, BehaviorHost->GetActiveModifiers(), MouseButton);
+}
+
+bool UScriptableToolClickSequenceBehavior::CanBeginClickSequence(const FInputDeviceRay& ClickPos)
+{
+	bool bShouldBegin = false;
+	if (CanBeginClickSequenceFunc.IsBound())
+	{
+		bShouldBegin = CanBeginClickSequenceFunc.Execute(ClickPos, BehaviorHost->GetActiveModifiers(), MouseButton);
+	}
+	return bShouldBegin;
+}
+
+void UScriptableToolClickSequenceBehavior::OnBeginClickSequence(const FInputDeviceRay& ClickPos)
+{
+	OnBeginClickSequenceFunc.ExecuteIfBound(ClickPos, BehaviorHost->GetActiveModifiers(), MouseButton);
+
+}
+
+void UScriptableToolClickSequenceBehavior::OnNextSequencePreview(const FInputDeviceRay& ClickPos)
+{
+	OnNextSequencePreviewFunc.ExecuteIfBound(ClickPos, BehaviorHost->GetActiveModifiers(), MouseButton);
+}
+
+bool UScriptableToolClickSequenceBehavior::OnNextSequenceClick(const FInputDeviceRay& ClickPos)
+{
+	bool bShouldContinue = false;
+	if (OnNextSequenceClickFunc.IsBound())
+	{
+		bShouldContinue = OnNextSequenceClickFunc.Execute(ClickPos, BehaviorHost->GetActiveModifiers(), MouseButton);
+	}
+	return bShouldContinue;
+}
+
+void UScriptableToolClickSequenceBehavior::OnTerminateClickSequence()
+{
+	OnTerminateClickSequenceFunc.ExecuteIfBound(BehaviorHost->GetActiveModifiers(), MouseButton);
+}
+
+bool UScriptableToolClickSequenceBehavior::RequestAbortClickSequence()
+{
+	bool bShouldTerminate = false;
+	if (RequestAbortClickSequenceFunc.IsBound())
+	{
+		bShouldTerminate = RequestAbortClickSequenceFunc.Execute();
+	}
+	return bShouldTerminate;
+}
+
+void UScriptableToolClickSequenceBehavior::OnUpdateModifierState(int ModifierID, bool bIsOn)
+{
+	BehaviorHost->OnUpdateModifierState(ModifierID, bIsOn);
+}
diff --git a/Engine/Plugins/Runtime/ScriptableToolsFramework/Source/ScriptableToolsFramework/Private/Behaviors/ScriptableToolMultiClickSequenceBehavior.h b/Engine/Plugins/Runtime/ScriptableToolsFramework/Source/ScriptableToolsFramework/Private/Behaviors/ScriptableToolMultiClickSequenceBehavior.h
new file mode 100644
--- /dev/null
+++ b/Engine/Plugins/Runtime/ScriptableToolsFramework/Source/ScriptableToolsFramework/Private/Behaviors/ScriptableToolMultiClickSequenceBehavior.h
@@ -0,0 +1,79 @@
+// Copyright Epic Games, Inc. All Rights Reserved.
+
+#pragma once
+
+#include "Behaviors/ScriptableToolBehaviorDelegates.h"
+#include "ScriptableToolBehavior.h"
+#include "BaseBehaviors/BehaviorTargetInterfaces.h"
+
+#include "ScriptableToolMultiClickSequenceBehavior.generated.h"
+
+class UScriptableModularBehaviorTool;
+class UMultiClickSequenceInputBehavior;
+
+
+UCLASS()
+class UScriptableToolClickSequenceBehavior : public UScriptableToolBehavior, public IClickSequenceBehaviorTarget
+{
+	GENERATED_BODY()
+
+public:
+	UScriptableToolClickSequenceBehavior() {};
+
+	void Init(TObjectPtr<UScriptableModularBehaviorTool> BehaviorHostIn,
+		FMouseBehaviorModiferCheckDelegate ModifierCheckFuncIn,
+		FOnBeginSequencePreviewDelegate OnBeginSequencePreviewFuncIn,
+		FCanBeginClickSequenceDelegate CanBeginClickSequenceFuncIn,
+		FOnBeginClickSequenceDelegate OnBeginClickSequenceFuncIn,
+		FOnNextSequencePreviewDelegate OnNextSequencePreviewFuncIn,
+		FOnNextSequenceClickDelegate OnNextSequenceClickFuncIn,
+		FOnTerminateClickSequenceDelegate OnTerminateClickSequenceFuncIn,
+		FRequestAbortClickSequenceDelegate RequestAbortClickSequenceFuncIn,
+		EScriptableToolMouseButton MouseButtonIn);
+
+	virtual UInputBehavior* GetWrappedBehavior() override;
+
+	virtual void OnBeginSequencePreview(const FInputDeviceRay& ClickPos) override;
+	virtual bool CanBeginClickSequence(const FInputDeviceRay& ClickPos) override;
+	virtual void OnBeginClickSequence(const FInputDeviceRay& ClickPos) override;
+	virtual void OnNextSequencePreview(const FInputDeviceRay& ClickPos) override;
+	virtual bool OnNextSequenceClick(const FInputDeviceRay& ClickPos) override;
+	virtual void OnTerminateClickSequence() override;
+	virtual bool RequestAbortClickSequence()  override;
+	virtual void OnUpdateModifierState(int ModifierID, bool bIsOn) override;
+
+private:
+
+	UPROPERTY()
+	TObjectPtr<UScriptableModularBehaviorTool> BehaviorHost;
+
+	UPROPERTY()
+	TObjectPtr<UMultiClickSequenceInputBehavior> Behavior;
+
+	UPROPERTY()
+	FMouseBehaviorModiferCheckDelegate ModifierCheckFunc;
+
+	UPROPERTY()
+	FOnBeginSequencePreviewDelegate OnBeginSequencePreviewFunc;
+
+	UPROPERTY()
+	FCanBeginClickSequenceDelegate CanBeginClickSequenceFunc;
+
+	UPROPERTY()
+	FOnBeginClickSequenceDelegate OnBeginClickSequenceFunc;
+
+	UPROPERTY()
+	FOnNextSequencePreviewDelegate OnNextSequencePreviewFunc;
+
+	UPROPERTY()
+	FOnNextSequenceClickDelegate OnNextSequenceClickFunc;
+
+	UPROPERTY()
+	FOnTerminateClickSequenceDelegate OnTerminateClickSequenceFunc;
+
+	UPROPERTY()
+	FRequestAbortClickSequenceDelegate RequestAbortClickSequenceFunc;
+
+	UPROPERTY()
+	EScriptableToolMouseButton MouseButton;
+};
\ No newline at end of file
diff --git a/Engine/Plugins/Runtime/ScriptableToolsFramework/Source/ScriptableToolsFramework/Private/Behaviors/ScriptableToolSingleClickBehavior.cpp b/Engine/Plugins/Runtime/ScriptableToolsFramework/Source/ScriptableToolsFramework/Private/Behaviors/ScriptableToolSingleClickBehavior.cpp
new file mode 100644
--- /dev/null
+++ b/Engine/Plugins/Runtime/ScriptableToolsFramework/Source/ScriptableToolsFramework/Private/Behaviors/ScriptableToolSingleClickBehavior.cpp
@@ -0,0 +1,87 @@
+// Copyright Epic Games, Inc. All Rights Reserved.
+
+#include "Behaviors/ScriptableToolSingleClickBehavior.h"
+
+#include "BaseTools/ScriptableModularBehaviorTool.h"
+#include "BaseBehaviors/SingleClickBehavior.h"
+
+void UScriptableToolSingleClickBehavior::Init(
+	TObjectPtr<UScriptableModularBehaviorTool> BehaviorHostIn,
+	FMouseBehaviorModiferCheckDelegate ModifierCheckFuncIn,
+	FTestIfHitByClickDelegate TestIfHitByClickFuncIn,
+	FOnHitByClickDelegate OnHitByClickFuncIn,
+	EScriptableToolMouseButton MouseButtonIn,
+	bool bHitTestOnReleaseIn
+	)
+{
+	BehaviorHost = BehaviorHostIn;
+	Behavior = NewObject<USingleClickInputBehavior>();
+	ModifierCheckFunc = ModifierCheckFuncIn;
+	TestIfHitByClickFunc = TestIfHitByClickFuncIn;
+	OnHitByClickFunc = OnHitByClickFuncIn;
+	MouseButton = MouseButtonIn;
+
+	Behavior->Initialize(this);
+	Behavior->ModifierCheckFunc = [this](const FInputDeviceState& InputDeviceState) {
+		bool bResult = true;
+		if (ModifierCheckFunc.IsBound())
+		{
+			bResult = ModifierCheckFunc.Execute(InputDeviceState);
+		}
+		return bResult;
+	};
+
+	Behavior->HitTestOnRelease = bHitTestOnReleaseIn;
+
+	BehaviorHost->AddInputBehavior(Behavior);
+
+	Behavior->Modifiers.RegisterModifier(1, FInputDeviceState::IsShiftKeyDown);
+	Behavior->Modifiers.RegisterModifier(2, FInputDeviceState::IsCtrlKeyDown);
+	Behavior->Modifiers.RegisterModifier(3, FInputDeviceState::IsAltKeyDown);
+
+	switch (MouseButton)
+	{
+	case EScriptableToolMouseButton::LeftButton:
+	{
+		Behavior->SetUseLeftMouseButton();
+		break;
+	}
+	case EScriptableToolMouseButton::RightButton:
+	{
+		Behavior->SetUseRightMouseButton();
+		break;
+	}
+	case EScriptableToolMouseButton::MiddleButton:
+	{
+		Behavior->SetUseMiddleMouseButton();
+		break;
+	}
+	default:
+		ensure(false);
+	}
+}
+
+UInputBehavior* UScriptableToolSingleClickBehavior::GetWrappedBehavior()
+{
+	return Behavior;
+}
+
+FInputRayHit UScriptableToolSingleClickBehavior::IsHitByClick(const FInputDeviceRay& ClickPos)
+{
+	FInputRayHit RayHit;
+	if (TestIfHitByClickFunc.IsBound())
+	{
+		RayHit = TestIfHitByClickFunc.Execute(ClickPos, BehaviorHost->GetActiveModifiers(), MouseButton);
+	}
+	return RayHit;
+}
+
+void UScriptableToolSingleClickBehavior::OnClicked(const FInputDeviceRay& ClickPos)
+{
+	OnHitByClickFunc.ExecuteIfBound(ClickPos, BehaviorHost->GetActiveModifiers(), MouseButton);
+}
+
+void UScriptableToolSingleClickBehavior::OnUpdateModifierState(int ModifierID, bool bIsOn)
+{
+	BehaviorHost->OnUpdateModifierState(ModifierID, bIsOn);
+}
diff --git a/Engine/Plugins/Runtime/ScriptableToolsFramework/Source/ScriptableToolsFramework/Private/Behaviors/ScriptableToolSingleClickBehavior.h b/Engine/Plugins/Runtime/ScriptableToolsFramework/Source/ScriptableToolsFramework/Private/Behaviors/ScriptableToolSingleClickBehavior.h
new file mode 100644
--- /dev/null
+++ b/Engine/Plugins/Runtime/ScriptableToolsFramework/Source/ScriptableToolsFramework/Private/Behaviors/ScriptableToolSingleClickBehavior.h
@@ -0,0 +1,54 @@
+// Copyright Epic Games, Inc. All Rights Reserved.
+
+#pragma once
+
+#include "Behaviors/ScriptableToolBehaviorDelegates.h"
+#include "ScriptableToolBehavior.h"
+#include "BaseBehaviors/BehaviorTargetInterfaces.h"
+
+#include "ScriptableToolSingleClickBehavior.generated.h"
+
+class UScriptableModularBehaviorTool;
+class USingleClickInputBehavior;
+
+UCLASS()
+class UScriptableToolSingleClickBehavior : public UScriptableToolBehavior, public IClickBehaviorTarget
+{
+	GENERATED_BODY()
+
+public:
+	UScriptableToolSingleClickBehavior() {};
+
+	void Init(TObjectPtr<UScriptableModularBehaviorTool> BehaviorHostIn,
+		FMouseBehaviorModiferCheckDelegate ModifierCheckFuncIn,
+		FTestIfHitByClickDelegate TestIfHitByClickFuncIn,
+		FOnHitByClickDelegate OnHitByClickFuncIn,
+		EScriptableToolMouseButton MouseButtonIn,
+		bool bHitTestOnReleaseIn);
+
+	virtual UInputBehavior* GetWrappedBehavior() override;
+
+	virtual FInputRayHit IsHitByClick(const FInputDeviceRay& ClickPos) override;
+	virtual void OnClicked(const FInputDeviceRay& ClickPos) override;
+	virtual void OnUpdateModifierState(int ModifierID, bool bIsOn) override;
+
+private:
+
+	UPROPERTY()
+	TObjectPtr<UScriptableModularBehaviorTool> BehaviorHost;
+
+	UPROPERTY()
+	TObjectPtr<USingleClickInputBehavior> Behavior;
+
+	UPROPERTY()
+	FMouseBehaviorModiferCheckDelegate ModifierCheckFunc;
+
+	UPROPERTY()
+	FTestIfHitByClickDelegate	TestIfHitByClickFunc;
+
+	UPROPERTY()
+	FOnHitByClickDelegate OnHitByClickFunc;
+
+	UPROPERTY()
+	EScriptableToolMouseButton MouseButton;
+};
diff --git a/Engine/Plugins/Runtime/ScriptableToolsFramework/Source/ScriptableToolsFramework/Private/Behaviors/ScriptableToolSingleClickOrDragBehavior.cpp b/Engine/Plugins/Runtime/ScriptableToolsFramework/Source/ScriptableToolsFramework/Private/Behaviors/ScriptableToolSingleClickOrDragBehavior.cpp
new file mode 100644
--- /dev/null
+++ b/Engine/Plugins/Runtime/ScriptableToolsFramework/Source/ScriptableToolsFramework/Private/Behaviors/ScriptableToolSingleClickOrDragBehavior.cpp
@@ -0,0 +1,128 @@
+// Copyright Epic Games, Inc. All Rights Reserved.
+
+#include "Behaviors/ScriptableToolSingleClickOrDragBehavior.h"
+
+#include "BaseTools/ScriptableModularBehaviorTool.h"
+#include "BaseBehaviors/SingleClickOrDragBehavior.h"
+
+
+void UScriptableToolSingleClickOrDragBehavior::Init(TObjectPtr<UScriptableModularBehaviorTool> BehaviorHostIn,
+	FMouseBehaviorModiferCheckDelegate ModifierCheckFuncIn,
+	FTestIfHitByClickDelegate TestIfHitByClickFuncIn,
+	FOnHitByClickDelegate OnHitByClickFuncIn,
+	FTestCanBeginClickDragSequenceDelegate TestCanBeginClickDragSequenceFuncIn,
+	FOnClickPressDelegate OnClickPressFuncIn,
+	FOnClickDragDelegate OnClickDragFuncIn,
+	FOnClickReleaseDelegate OnClickReleaseFuncIn,
+	FOnTerminateDragSequenceDelegate OnTerminateDragSequenceFuncIn,
+	EScriptableToolMouseButton MouseButtonIn,
+	bool bBeginDragIfClickTargetNotHitIn,
+	float ClickDistanceThresholdIn)
+{
+	BehaviorHost = BehaviorHostIn;
+	Behavior = NewObject<USingleClickOrDragInputBehavior>();
+	ModifierCheckFunc = ModifierCheckFuncIn;
+	TestIfHitByClickFunc = TestIfHitByClickFuncIn;
+	OnHitByClickFunc = OnHitByClickFuncIn;
+	TestCanBeginClickDragSequenceFunc = TestCanBeginClickDragSequenceFuncIn;
+	OnClickPressFunc = OnClickPressFuncIn;
+	OnClickDragFunc = OnClickDragFuncIn;
+	OnClickReleaseFunc = OnClickReleaseFuncIn;
+	OnTerminateDragSequenceFunc = OnTerminateDragSequenceFuncIn;
+	MouseButton = MouseButtonIn;
+
+	Behavior->Initialize(this, this);
+	Behavior->ModifierCheckFunc = [this](const FInputDeviceState& InputDeviceState) {
+		bool bResult = true;
+		if (ModifierCheckFunc.IsBound())
+		{
+			bResult = ModifierCheckFunc.Execute(InputDeviceState);
+		}
+		return bResult;
+	};
+
+	Behavior->bBeginDragIfClickTargetNotHit = bBeginDragIfClickTargetNotHitIn;
+	Behavior->ClickDistanceThreshold = ClickDistanceThresholdIn;
+
+	BehaviorHost->AddInputBehavior(Behavior);
+
+	Behavior->Modifiers.RegisterModifier(1, FInputDeviceState::IsShiftKeyDown);
+	Behavior->Modifiers.RegisterModifier(2, FInputDeviceState::IsCtrlKeyDown);
+	Behavior->Modifiers.RegisterModifier(3, FInputDeviceState::IsAltKeyDown);
+
+	switch (MouseButton)
+	{
+	case EScriptableToolMouseButton::LeftButton:
+	{
+		Behavior->SetUseLeftMouseButton();
+		break;
+	}
+	case EScriptableToolMouseButton::RightButton:
+	{
+		Behavior->SetUseRightMouseButton();
+		break;
+	}
+	case EScriptableToolMouseButton::MiddleButton:
+	{
+		Behavior->SetUseMiddleMouseButton();
+		break;
+	}
+	default:
+		ensure(false);
+	}
+}
+
+UInputBehavior* UScriptableToolSingleClickOrDragBehavior::GetWrappedBehavior()
+{
+	return Behavior;
+}
+
+FInputRayHit UScriptableToolSingleClickOrDragBehavior::IsHitByClick(const FInputDeviceRay& ClickPos)
+{
+	FInputRayHit RayHit;
+	if (TestIfHitByClickFunc.IsBound())
+	{
+		RayHit = TestIfHitByClickFunc.Execute(ClickPos, BehaviorHost->GetActiveModifiers(), MouseButton);
+	}
+	return RayHit;
+}
+
+void UScriptableToolSingleClickOrDragBehavior::OnClicked(const FInputDeviceRay& ClickPos)
+{
+	OnHitByClickFunc.ExecuteIfBound(ClickPos, BehaviorHost->GetActiveModifiers(), MouseButton);
+}
+
+FInputRayHit UScriptableToolSingleClickOrDragBehavior::CanBeginClickDragSequence(const FInputDeviceRay& PressPos)
+{
+	FInputRayHit RayHit;
+	if (TestCanBeginClickDragSequenceFunc.IsBound())
+	{
+		RayHit = TestCanBeginClickDragSequenceFunc.Execute(PressPos, BehaviorHost->GetActiveModifiers(), MouseButton);
+	}
+	return RayHit;
+}
+
+void UScriptableToolSingleClickOrDragBehavior::OnClickPress(const FInputDeviceRay& PressPos)
+{
+	OnClickPressFunc.ExecuteIfBound(PressPos, BehaviorHost->GetActiveModifiers(), MouseButton);
+}
+
+void UScriptableToolSingleClickOrDragBehavior::OnClickDrag(const FInputDeviceRay& DragPos)
+{
+	OnClickDragFunc.ExecuteIfBound(DragPos, BehaviorHost->GetActiveModifiers(), MouseButton);
+}
+
+void UScriptableToolSingleClickOrDragBehavior::OnClickRelease(const FInputDeviceRay& ReleasePos)
+{
+	OnClickReleaseFunc.ExecuteIfBound(ReleasePos, BehaviorHost->GetActiveModifiers(), MouseButton);
+}
+
+void UScriptableToolSingleClickOrDragBehavior::OnTerminateDragSequence()
+{
+	OnTerminateDragSequenceFunc.ExecuteIfBound(BehaviorHost->GetActiveModifiers(), MouseButton);
+}
+
+void UScriptableToolSingleClickOrDragBehavior::OnUpdateModifierState(int ModifierID, bool bIsOn)
+{
+	BehaviorHost->OnUpdateModifierState(ModifierID, bIsOn);
+}
diff --git a/Engine/Plugins/Runtime/ScriptableToolsFramework/Source/ScriptableToolsFramework/Private/Behaviors/ScriptableToolSingleClickOrDragBehavior.h b/Engine/Plugins/Runtime/ScriptableToolsFramework/Source/ScriptableToolsFramework/Private/Behaviors/ScriptableToolSingleClickOrDragBehavior.h
new file mode 100644
--- /dev/null
+++ b/Engine/Plugins/Runtime/ScriptableToolsFramework/Source/ScriptableToolsFramework/Private/Behaviors/ScriptableToolSingleClickOrDragBehavior.h
@@ -0,0 +1,81 @@
+// Copyright Epic Games, Inc. All Rights Reserved.
+
+#pragma once
+
+#include "Behaviors/ScriptableToolBehaviorDelegates.h"
+#include "ScriptableToolBehavior.h"
+#include "BaseBehaviors/BehaviorTargetInterfaces.h"
+
+#include "ScriptableToolSingleClickOrDragBehavior.generated.h"
+
+class UScriptableModularBehaviorTool;
+class USingleClickOrDragInputBehavior;
+
+
+UCLASS()
+class UScriptableToolSingleClickOrDragBehavior : public UScriptableToolBehavior, public IClickDragBehaviorTarget, public IClickBehaviorTarget
+{
+	GENERATED_BODY()
+
+public:
+	UScriptableToolSingleClickOrDragBehavior() {};
+
+	void Init(TObjectPtr<UScriptableModularBehaviorTool> BehaviorHostIn,
+		FMouseBehaviorModiferCheckDelegate ModifierCheckFuncIn,
+		FTestIfHitByClickDelegate TestIfHitByClickFuncIn,
+		FOnHitByClickDelegate OnHitByClickFuncIn,
+		FTestCanBeginClickDragSequenceDelegate TestCanBeginClickDragSequenceFuncIn,
+		FOnClickPressDelegate OnClickPressFuncIn,
+		FOnClickDragDelegate OnClickDragFuncIn,
+		FOnClickReleaseDelegate OnClickReleaseFuncIn,
+		FOnTerminateDragSequenceDelegate OnTerminateDragSequenceFuncIn,
+		EScriptableToolMouseButton MouseButtonIn,
+		bool bBeginDragIfClickTargetNotHitIn,
+		float ClickDistanceThresholdIn);
+
+	virtual UInputBehavior* GetWrappedBehavior() override;
+
+	virtual FInputRayHit IsHitByClick(const FInputDeviceRay& ClickPos) override;
+	virtual void OnClicked(const FInputDeviceRay& ClickPos) override;
+	virtual FInputRayHit CanBeginClickDragSequence(const FInputDeviceRay& PressPos) override;
+	virtual void OnClickPress(const FInputDeviceRay& PressPos) override;
+	virtual void OnClickDrag(const FInputDeviceRay& DragPos) override;
+	virtual void OnClickRelease(const FInputDeviceRay& ReleasePos) override;
+	virtual void OnTerminateDragSequence() override;
+	virtual void OnUpdateModifierState(int ModifierID, bool bIsOn) override;
+
+private:
+
+	UPROPERTY()
+	TObjectPtr<UScriptableModularBehaviorTool> BehaviorHost;
+
+	UPROPERTY()
+	TObjectPtr<USingleClickOrDragInputBehavior> Behavior;
+
+	UPROPERTY()
+	FMouseBehaviorModiferCheckDelegate ModifierCheckFunc;
+
+	UPROPERTY()
+	FTestIfHitByClickDelegate	TestIfHitByClickFunc;
+
+	UPROPERTY()
+	FOnHitByClickDelegate OnHitByClickFunc;
+
+	UPROPERTY()
+	FTestCanBeginClickDragSequenceDelegate	TestCanBeginClickDragSequenceFunc;
+
+	UPROPERTY()
+	FOnClickPressDelegate OnClickPressFunc;
+
+	UPROPERTY()
+	FOnClickDragDelegate OnClickDragFunc;
+
+	UPROPERTY()
+	FOnClickReleaseDelegate OnClickReleaseFunc;
+
+	UPROPERTY()
+	FOnTerminateDragSequenceDelegate OnTerminateDragSequenceFunc;
+
+	UPROPERTY()
+	EScriptableToolMouseButton MouseButton;
+};
\ No newline at end of file
diff --git a/Engine/Plugins/Runtime/ScriptableToolsFramework/Source/ScriptableToolsFramework/Public/BaseTools/ScriptableModularBehaviorTool.h b/Engine/Plugins/Runtime/ScriptableToolsFramework/Source/ScriptableToolsFramework/Public/BaseTools/ScriptableModularBehaviorTool.h
new file mode 100644
--- /dev/null
+++ b/Engine/Plugins/Runtime/ScriptableToolsFramework/Source/ScriptableToolsFramework/Public/BaseTools/ScriptableModularBehaviorTool.h
@@ -0,0 +1,275 @@
+// Copyright Epic Games, Inc. All Rights Reserved.
+
+#pragma once
+
+#include "ScriptableInteractiveTool.h"
+#include "InputState.h"
+#include "InputBehavior.h"
+#include "Behaviors/ScriptableToolBehaviorDelegates.h"
+#include "ScriptableModularBehaviorTool.generated.h"
+
+class UScriptableToolSingleClickBehavior;
+class UScriptableToolClickDragBehavior;
+class UScriptableToolSingleClickOrDragBehavior;
+class UScriptableToolMouseWheelBehavior;
+class UScriptableToolClickSequenceBehavior;
+class UScriptableToolMouseHoverBehavior;
+
+/**
+ A Scriptable tool base blueprint class which provides support for user defined mouse interaction behaviors
+ */ 
+UCLASS()
+class SCRIPTABLETOOLSFRAMEWORK_API UScriptableModularBehaviorTool :	public UScriptableInteractiveTool
+{
+	GENERATED_BODY()
+
+protected:
+
+
+public:	
+
+	/**
+	*   Implements a standard "button-click"-style input behavior
+	*   The state machine works as follows:
+	*    1) on input-device-button-press, hit-test the target. If hit, begin capture
+	*    2) on input-device-button-release, hit-test the target. If hit, call OnClicked(). If not hit, ignore click.
+	*    
+	*   The second hit-test is required to allow the click to be "cancelled" by moving away
+	*   from the target. This is standard GUI behavior. You can disable this second hit test
+	*   using the HitTestOnRelease property. This is strongly discouraged.
+	* 
+	*	@param IfHitByClick Test if hit by a click
+	*	@param OnHitByClick Notify that click occurred
+	*	@param CaptureCheck Only enable capture if returns true
+	*	@param CapturePriority The priority is used to resolve situations where multiple behaviors want the same capture
+	*	@param MouseButton Determines which mouse button the behavior captures on
+	*/
+	UFUNCTION(BlueprintCallable, Category = "ScriptableTool|Input", meta = (AdvancedDisplay = "CaptureCheck, CapturePriority, MouseButton, bHitTestOnRelease"))
+	void AddSingleClickBehavior(
+		const FTestIfHitByClickDelegate& IfHitByClick,
+		const FOnHitByClickDelegate& OnHitByClick,
+		const FMouseBehaviorModiferCheckDelegate& CaptureCheck,
+		int CapturePriority = 100,
+		EScriptableToolMouseButton MouseButton = EScriptableToolMouseButton::LeftButton,
+		bool bHitTestOnRelease = true
+	);
+
+	/**
+	*   Implements a standard "button-click-drag"-style input behavior.
+	*	 
+	*   The state machine works as follows:
+	*      1) on input-device-button-press, call CanBeginClickDragSequence to determine if capture should begin
+	*      2) on input-device-move, call OnClickDrag
+	*      3) on input-device-button-release, call OnClickRelease
+	*    
+	*   If a ForceEndCapture occurs we call OnTerminateDragSequence   
+	* 
+	*   @param CanBeginClickDragSequence Test if target can begin click-drag interaction at this point
+	*	@param OnClickPress Notify Target that click press occurred
+	*	@param OnClickDrag Notify Target that input position has changed
+	*	@param OnClickRelease Notify Target that click release occurred
+	*	@param OnTerminateDragSequence Notify Target that click-drag sequence has been explicitly terminated (eg by escape key)
+	*	@param CaptureCheck Only enable capture if returns true
+	*	@param CapturePriority The priority is used to resolve situations where multiple behaviors want the same capture
+	*	@param MouseButton Determines which mouse button the behavior captures on
+	*/
+	UFUNCTION(BlueprintCallable, Category = "ScriptableTool|Input", meta = (AdvancedDisplay = "CaptureCheck, CapturePriority, MouseButton, bUpdateModifiersDuringDrag"))
+	void AddClickDragBehavior(		
+		const FTestCanBeginClickDragSequenceDelegate& CanBeginClickDragSequence,
+		const FOnClickPressDelegate& OnClickPress,
+		const FOnClickDragDelegate& OnClickDrag,
+		const FOnClickReleaseDelegate& OnClickRelease,
+		const FOnTerminateDragSequenceDelegate& OnTerminateDragSequence,
+		const FMouseBehaviorModiferCheckDelegate& CaptureCheck,
+		int CapturePriority = 100,
+		EScriptableToolMouseButton MouseButton = EScriptableToolMouseButton::LeftButton,
+		bool bUpdateModifiersDuringDrag = false
+	);
+
+	/**
+	* 
+	*   SingleClickOrDragBehavior is a combination of a SingleClickBehavior and ClickDragBehavior,
+	*   and allows for the common UI interaction where a click-and-release does one action, but if the mouse 
+	*   is moved, then a drag interaction is started. For example click-to-select is often combined with
+	*   a drag-marquee-rectangle in this way. This can be directly implemented with a ClickDragBehavior but
+	*   requires the client to (eg) detect movement thresholds, etc. This class encapsulates all that state/logic.
+	* 
+	*   The ClickDistanceThreshold parameter determines how far the mouse must move (in whatever device units are in use)
+	*   to switch from a click to drag interaction
+	* 
+	*   The bBeginDragIfClickTargetNotHit parameter determines if the drag interaction will be immediately initiated 
+	*   if the initial 'click' mouse-down does not hit a valid clickable target. Defaults to true. 
+	* 
+	*	@param IfHitByClick Test if hit by a click
+	*	@param OnHitByClick Notify that click occurred
+	*   @param CanBeginClickDragSequence Test if target can begin click-drag interaction at this point
+	*	@param OnClickPress Notify Target that click press occurred
+	*	@param OnClickDrag Notify Target that input position has changed
+	*	@param OnClickRelease Notify Target that click release occurred
+	*	@param OnTerminateDragSequence Notify Target that click-drag sequence has been explicitly terminated (eg by escape key)
+	*	@param CaptureCheck Only enable capture if returns true
+	*	@param CapturePriority The priority is used to resolve situations where multiple behaviors want the same capture
+	*	@param MouseButton Determines which mouse button the behavior captures on
+	*/
+	UFUNCTION(BlueprintCallable, Category = "ScriptableTool|Input", meta = (AdvancedDisplay = "CaptureCheck, CapturePriority, MouseButton, bBeginDragIfClickTargetNotHit, ClickDistanceThreshold"))
+	void AddSingleClickOrDragBehavior(		
+		FTestIfHitByClickDelegate IfHitByClick,
+		FOnHitByClickDelegate OnHitByClick,
+		FTestCanBeginClickDragSequenceDelegate CanBeginClickDragSequence,
+		FOnClickPressDelegate OnClickPress,
+		FOnClickDragDelegate OnClickDrag,
+		FOnClickReleaseDelegate OnClickRelease,
+		FOnTerminateDragSequenceDelegate OnTerminateDragSequence,
+		FMouseBehaviorModiferCheckDelegate CaptureCheck,
+		int CapturePriority = 100,
+		EScriptableToolMouseButton MouseButton = EScriptableToolMouseButton::LeftButton,
+		bool bBeginDragIfClickTargetNotHit = true,
+		float ClickDistanceThreshold = 5.0
+	);
+
+	/**
+	*	@param TestShouldRespondToMouseWheel The result's bHit property determines whether the mouse wheel action will be captured. (Perhaps the mouse wheel only does something when mousing over some part of a mesh)
+	*	@param OnMouseWheelScrollUp CurrentPos device position/ray at point where mouse wheel is engaged
+	*	@param OnMouseWheelScrollDown CurrentPos device position/ray at point where mouse wheel is engaged
+	*	@param CaptureCheck Only enable capture if returns true
+	*	@param CapturePriority The priority is used to resolve situations where multiple behaviors want the same capture
+	*/
+	UFUNCTION(BlueprintCallable, Category = "ScriptableTool|Input", meta = (AdvancedDisplay = "CaptureCheck, CapturePriority"))
+	void AddMouseWheelBehavior(		
+		FTestShouldRespondToMouseWheelDelegate TestShouldRespondToMouseWheel,
+		FOnMouseWheelScrollUpDelegate OnMouseWheelScrollUp,
+		FOnMouseWheelScrollDownDelegate OnMouseWheelScrollDown,
+		FMouseBehaviorModiferCheckDelegate CaptureCheck,
+		int CapturePriority = 100
+	);
+
+	/**
+	* 
+	*   MultiClickSequenceBehavior implements a generic multi-click-sequence input behavior.
+	*   For example this behavior could be used to implement a multi-click polygon-drawing interaction.
+	*
+	*   The internal state machine works as follows:
+	*     1) on input-device-button-press, check if target wants to begin sequence. If so, begin capture.
+	*     2) on button *release*, check if target wants to continue or terminate sequence
+	*         a) if terminate, release capture
+	*         b) if continue, do nothing (capture continues between presses)
+	*
+	*   The target will receive "preview" notifications (basically hover) during updates where there is
+	*   not a release. This can be used to (eg) update a rubber-band selection end point
+	* 
+	*	@param OnBeginSequencePreview Notify Target device position has changed but a click sequence hasn't begun yet (eg for interactive previews)
+	*	@param CanBeginClickSequence Test if target would like to begin sequence based on this click. Gets checked both on mouse down and mouse up.
+	*	@param OnBeginClickSequence Notify Target about the first click in the sequence.
+	*	@param OnNextSequencePreview Notify Target device position has changed but a click hasn't occurred yet (eg for interactive previews)
+	*	@param OnNextSequenceClick  Notify Target about next click in sqeuence, returns false to terminate sequence
+	*	@param OnTerminateClickSequence Notify Target that click sequence has been explicitly terminated (eg by escape key, cancel tool, etc).
+	*	@param RequestAbortClickSequence Target overrides this and returns true if it wants to abort click sequence, checked every update
+	*	@param CaptureCheck Only enable capture if returns true
+	*	@param CapturePriority The priority is used to resolve situations where multiple behaviors want the same capture
+	*	@param MouseButton Determines which mouse button the behavior captures on
+	*/
+	UFUNCTION(BlueprintCallable, Category = "ScriptableTool|Input", meta = (AdvancedDisplay = "CaptureCheck, CapturePriority, MouseButton"))
+	void AddMultiClickSequenceBehavior(		
+		FOnBeginSequencePreviewDelegate OnBeginSequencePreview,
+		FCanBeginClickSequenceDelegate CanBeginClickSequence,
+		FOnBeginClickSequenceDelegate OnBeginClickSequence,
+		FOnNextSequencePreviewDelegate OnNextSequencePreview,
+		FOnNextSequenceClickDelegate OnNextSequenceClick,
+		FOnTerminateClickSequenceDelegate OnTerminateClickSequence,
+		FRequestAbortClickSequenceDelegate RequestAbortClickSequence,
+		FMouseBehaviorModiferCheckDelegate CaptureCheck,
+		int CapturePriority = 100,
+		EScriptableToolMouseButton MouseButton = EScriptableToolMouseButton::LeftButton
+	);
+
+	/**
+	*	@param BeginHoverSequenceHitTest  Do hover hit-test
+	*	@param OnBeginHover Initialize hover sequence at given position
+	*	@param OnUpdateHover Update active hover sequence with new input position
+	*	@param OnEndHover Terminate active hover sequence
+	*	@param CapturePriority The priority is used to resolve situations where multiple behaviors want the same capture
+	*/
+	UFUNCTION(BlueprintCallable, Category = "ScriptableTool|Input", meta = (AdvancedDisplay = "CapturePriority"))
+	void AddMouseHoverBehavior(
+		FBeginHoverSequenceHitTestDelegate BeginHoverSequenceHitTest,
+		FOnBeginHoverDelegate OnBeginHover,
+		FOnUpdateHoverDelegate OnUpdateHover,
+		FOnEndHoverDelegate OnEndHover,
+		int CapturePriority = 100
+	);
+
+
+private:
+
+	UPROPERTY(Transient, DuplicateTransient, NonTransactional, SkipSerialization)
+	TArray< TObjectPtr<UScriptableToolSingleClickBehavior> > SingleClickBehaviors;
+
+	UPROPERTY(Transient, DuplicateTransient, NonTransactional, SkipSerialization)
+	TArray< TObjectPtr<UScriptableToolClickDragBehavior> > ClickDragBehaviors;
+
+	UPROPERTY(Transient, DuplicateTransient, NonTransactional, SkipSerialization)
+	TArray< TObjectPtr<UScriptableToolSingleClickOrDragBehavior> > SingleClickOrDragBehaviors;
+
+	UPROPERTY(Transient, DuplicateTransient, NonTransactional, SkipSerialization)
+	TArray< TObjectPtr<UScriptableToolMouseWheelBehavior> > MouseWheelBehaviors;
+
+	UPROPERTY(Transient, DuplicateTransient, NonTransactional, SkipSerialization)
+	TArray< TObjectPtr<UScriptableToolClickSequenceBehavior> > MultiClickSequenceBehaviors;
+
+	UPROPERTY(Transient, DuplicateTransient, NonTransactional, SkipSerialization)
+	TArray< TObjectPtr<UScriptableToolMouseHoverBehavior> > MouseHoverBehaviors;
+
+
+	//
+	// Modifer Buttons Support
+	// 
+	// We store these here, instead of in the behavior wrapper classes, to provide "global" access
+	// via the blueprint methods regardless of what, if any, behaviors have been added to the tool.
+	//
+public:
+
+	/** @return true if the Shift key is currently held down */
+	UFUNCTION(BlueprintCallable, Category="ScriptableTool|Input")
+	UPARAM(DisplayName="Shift Down") bool IsShiftDown() const;
+
+	/** @return true if the Ctrl key is currently held down */
+	UFUNCTION(BlueprintCallable, Category="ScriptableTool|Input")
+	UPARAM(DisplayName="Ctrl Down") bool IsCtrlDown() const;
+
+	/** @return true if the Alt key is currently held down */
+	UFUNCTION(BlueprintCallable, Category="ScriptableTool|Input")
+	UPARAM(DisplayName="Alt Down") bool IsAltDown() const;
+
+	/** @return a struct containing the current Modifier key states */
+	UFUNCTION(BlueprintCallable, Category="ScriptableTool|Input")
+	FScriptableToolModifierStates GetActiveModifiers();
+
+	void OnUpdateModifierState(int ModifierID, bool bIsOn);
+
+private:
+
+	bool bShiftModifier = false;
+	bool bCtrlModifier = false;
+	bool bAltModifier = false;
+
+
+};
+
+
+
+
+
+
+
+
+
+
+
+
+
+
+
+
+
+
+
diff --git a/Engine/Plugins/Runtime/ScriptableToolsFramework/Source/ScriptableToolsFramework/Public/Behaviors/ScriptableToolBehaviorDelegates.h b/Engine/Plugins/Runtime/ScriptableToolsFramework/Source/ScriptableToolsFramework/Public/Behaviors/ScriptableToolBehaviorDelegates.h
new file mode 100644
--- /dev/null
+++ b/Engine/Plugins/Runtime/ScriptableToolsFramework/Source/ScriptableToolsFramework/Public/Behaviors/ScriptableToolBehaviorDelegates.h
@@ -0,0 +1,44 @@
+// Copyright Epic Games, Inc. All Rights Reserved.
+
+#pragma once
+
+#include "ScriptableInteractiveTool.h"
+#include "InputState.h"
+
+#include "ScriptableToolBehaviorDelegates.generated.h"
+
+UENUM(BlueprintType)
+enum class EScriptableToolMouseButton : uint8
+{
+	LeftButton,
+	RightButton,
+	MiddleButton
+};
+
+DECLARE_DYNAMIC_DELEGATE_RetVal_OneParam(bool, FMouseBehaviorModiferCheckDelegate, const FInputDeviceState&, InputDeviceState);
+
+DECLARE_DYNAMIC_DELEGATE_RetVal_ThreeParams(FInputRayHit, FTestIfHitByClickDelegate, FInputDeviceRay, ClickPos, FScriptableToolModifierStates, Modifiers, EScriptableToolMouseButton, MouseButton);
+DECLARE_DYNAMIC_DELEGATE_ThreeParams(FOnHitByClickDelegate, FInputDeviceRay, ClickPos, FScriptableToolModifierStates, Modifiers, EScriptableToolMouseButton, MouseButton);
+
+DECLARE_DYNAMIC_DELEGATE_RetVal_ThreeParams(FInputRayHit, FTestCanBeginClickDragSequenceDelegate, FInputDeviceRay, PressPos, FScriptableToolModifierStates, Modifiers, EScriptableToolMouseButton, MouseButton);
+DECLARE_DYNAMIC_DELEGATE_ThreeParams(FOnClickPressDelegate, FInputDeviceRay, PressPos, FScriptableToolModifierStates, Modifiers, EScriptableToolMouseButton, MouseButton);
+DECLARE_DYNAMIC_DELEGATE_ThreeParams(FOnClickDragDelegate, FInputDeviceRay, DragPos, FScriptableToolModifierStates, Modifiers, EScriptableToolMouseButton, MouseButton);
+DECLARE_DYNAMIC_DELEGATE_ThreeParams(FOnClickReleaseDelegate, FInputDeviceRay, ReleasePos, FScriptableToolModifierStates, Modifiers, EScriptableToolMouseButton, MouseButton);
+DECLARE_DYNAMIC_DELEGATE_TwoParams(FOnTerminateDragSequenceDelegate, FScriptableToolModifierStates, Modifiers, EScriptableToolMouseButton, MouseButton);
+
+DECLARE_DYNAMIC_DELEGATE_RetVal_TwoParams(FInputRayHit, FTestShouldRespondToMouseWheelDelegate, FInputDeviceRay, CurrentPos, FScriptableToolModifierStates, Modifiers);
+DECLARE_DYNAMIC_DELEGATE_TwoParams(FOnMouseWheelScrollUpDelegate, FInputDeviceRay, CurrentPos, FScriptableToolModifierStates, Modifiers);
+DECLARE_DYNAMIC_DELEGATE_TwoParams(FOnMouseWheelScrollDownDelegate, FInputDeviceRay, CurrentPos, FScriptableToolModifierStates, Modifiers);
+
+DECLARE_DYNAMIC_DELEGATE_ThreeParams(FOnBeginSequencePreviewDelegate, FInputDeviceRay, ClickPos, FScriptableToolModifierStates, Modifiers, EScriptableToolMouseButton, MouseButton);
+DECLARE_DYNAMIC_DELEGATE_RetVal_ThreeParams(bool, FCanBeginClickSequenceDelegate, FInputDeviceRay, ClickPos, FScriptableToolModifierStates, Modifiers, EScriptableToolMouseButton, MouseButton);
+DECLARE_DYNAMIC_DELEGATE_ThreeParams(FOnBeginClickSequenceDelegate, FInputDeviceRay, ClickPos, FScriptableToolModifierStates, Modifiers, EScriptableToolMouseButton, MouseButton);
+DECLARE_DYNAMIC_DELEGATE_ThreeParams(FOnNextSequencePreviewDelegate, FInputDeviceRay, ClickPos, FScriptableToolModifierStates, Modifiers, EScriptableToolMouseButton, MouseButton);
+DECLARE_DYNAMIC_DELEGATE_RetVal_ThreeParams(bool, FOnNextSequenceClickDelegate, FInputDeviceRay, ClickPos, FScriptableToolModifierStates, Modifiers, EScriptableToolMouseButton, MouseButton);
+DECLARE_DYNAMIC_DELEGATE_TwoParams(FOnTerminateClickSequenceDelegate, FScriptableToolModifierStates, Modifiers, EScriptableToolMouseButton, MouseButton);
+DECLARE_DYNAMIC_DELEGATE_RetVal(bool, FRequestAbortClickSequenceDelegate);
+
+DECLARE_DYNAMIC_DELEGATE_RetVal_TwoParams(FInputRayHit, FBeginHoverSequenceHitTestDelegate, FInputDeviceRay, CurrentPos, FScriptableToolModifierStates, Modifiers);
+DECLARE_DYNAMIC_DELEGATE_TwoParams(FOnBeginHoverDelegate, FInputDeviceRay, CurrentPos, FScriptableToolModifierStates, Modifiers);
+DECLARE_DYNAMIC_DELEGATE_RetVal_TwoParams(bool, FOnUpdateHoverDelegate, FInputDeviceRay, CurrentPos, FScriptableToolModifierStates, Modifiers);
+DECLARE_DYNAMIC_DELEGATE(FOnEndHoverDelegate);
\ No newline at end of file
diff --git a/Engine/Plugins/Runtime/ScriptableToolsFramework/Source/ScriptableToolsFramework/ScriptableToolsFramework.Build.cs b/Engine/Plugins/Runtime/ScriptableToolsFramework/Source/ScriptableToolsFramework/ScriptableToolsFramework.Build.cs
--- a/Engine/Plugins/Runtime/ScriptableToolsFramework/Source/ScriptableToolsFramework/ScriptableToolsFramework.Build.cs
+++ b/Engine/Plugins/Runtime/ScriptableToolsFramework/Source/ScriptableToolsFramework/ScriptableToolsFramework.Build.cs
@@ -30,6 +30,7 @@ public class ScriptableToolsFramework : ModuleRules
 				"PhysicsCore",
 				"RenderCore",
 				"GeometryCore",
+				"InputCore",
 				"InteractiveToolsFramework"
 			}
 			);
