ScriptableTools: Fix missing copyright boilerplate

diff --git a/Engine/Plugins/Runtime/ScriptableToolsFramework/Source/ScriptableToolsFramework/Private/Behaviors/ScriptableToolBehavior.cpp b/Engine/Plugins/Runtime/ScriptableToolsFramework/Source/ScriptableToolsFramework/Private/Behaviors/ScriptableToolBehavior.cpp
--- a/Engine/Plugins/Runtime/ScriptableToolsFramework/Source/ScriptableToolsFramework/Private/Behaviors/ScriptableToolBehavior.cpp
+++ b/Engine/Plugins/Runtime/ScriptableToolsFramework/Source/ScriptableToolsFramework/Private/Behaviors/ScriptableToolBehavior.cpp
@@ -1,3 +1,5 @@
+// Copyright Epic Games, Inc. All Rights Reserved.
+
 #include "ScriptableToolBehavior.h"
 
 void UScriptableToolBehavior::SetDefaultPriority(const FInputCapturePriority& Priority)
