ScriptableTools: Fix various issues with the scriptable tool framework for behavior events.

* Fix issue where mouse wheel down event was not bound correctly
* Fix issue where pre-capture delegates were recieving modifier keys when they weren't available yet from the behavior.
* Fix issue where several behaviors took delegates by ref, preventing their optional setting by the user.

diff --git a/Engine/Plugins/Runtime/ScriptableToolsFramework/Source/ScriptableToolsFramework/Private/BaseTools/ScriptableModularBehaviorTool.cpp b/Engine/Plugins/Runtime/ScriptableToolsFramework/Source/ScriptableToolsFramework/Private/BaseTools/ScriptableModularBehaviorTool.cpp
--- a/Engine/Plugins/Runtime/ScriptableToolsFramework/Source/ScriptableToolsFramework/Private/BaseTools/ScriptableModularBehaviorTool.cpp
+++ b/Engine/Plugins/Runtime/ScriptableToolsFramework/Source/ScriptableToolsFramework/Private/BaseTools/ScriptableModularBehaviorTool.cpp
@@ -10,9 +10,9 @@
 #include "Behaviors/ScriptableToolSingleClickOrDragBehavior.h"
 
 void UScriptableModularBehaviorTool::AddSingleClickBehavior(
-	const FTestIfHitByClickDelegate& TestIfHitByClickDelegate,
-	const FOnHitByClickDelegate& OnHitByClickDelegate,
-	const FMouseBehaviorModiferCheckDelegate& ModifierCheckFunction,
+	const FTestIfHitByClickDelegate TestIfHitByClickDelegate,
+	const FOnHitByClickDelegate OnHitByClickDelegate,
+	const FMouseBehaviorModiferCheckDelegate ModifierCheckFunction,
 	int CapturePriority,
 	EScriptableToolMouseButton MouseButton,
 	bool bHitTestOnRelease
@@ -27,12 +27,12 @@ void UScriptableModularBehaviorTool::AddSingleClickBehavior(
 
 
 void UScriptableModularBehaviorTool::AddClickDragBehavior(
-	const FTestCanBeginClickDragSequenceDelegate& TestCanBeginClickDragSequenceFuncIn,
-	const FOnClickPressDelegate& OnClickPressFuncIn,
-	const FOnClickDragDelegate& OnClickDragFuncIn,
-	const FOnClickReleaseDelegate& OnClickReleaseFuncIn,
-	const FOnTerminateDragSequenceDelegate& OnTerminateDragSequenceFuncIn,
-	const FMouseBehaviorModiferCheckDelegate& ModifierCheckFuncIn,
+	const FTestCanBeginClickDragSequenceDelegate TestCanBeginClickDragSequenceFuncIn,
+	const FOnClickPressDelegate OnClickPressFuncIn,
+	const FOnClickDragDelegate OnClickDragFuncIn,
+	const FOnClickReleaseDelegate OnClickReleaseFuncIn,
+	const FOnTerminateDragSequenceDelegate OnTerminateDragSequenceFuncIn,
+	const FMouseBehaviorModiferCheckDelegate ModifierCheckFuncIn,
 	int CapturePriority,
 	EScriptableToolMouseButton MouseButtonIn,
 	bool bUpdateModifiersDuringDrag
diff --git a/Engine/Plugins/Runtime/ScriptableToolsFramework/Source/ScriptableToolsFramework/Private/Behaviors/ScriptableToolMouseWheelBehavior.cpp b/Engine/Plugins/Runtime/ScriptableToolsFramework/Source/ScriptableToolsFramework/Private/Behaviors/ScriptableToolMouseWheelBehavior.cpp
--- a/Engine/Plugins/Runtime/ScriptableToolsFramework/Source/ScriptableToolsFramework/Private/Behaviors/ScriptableToolMouseWheelBehavior.cpp
+++ b/Engine/Plugins/Runtime/ScriptableToolsFramework/Source/ScriptableToolsFramework/Private/Behaviors/ScriptableToolMouseWheelBehavior.cpp
@@ -48,7 +48,7 @@ FInputRayHit UScriptableToolMouseWheelBehavior::ShouldRespondToMouseWheel(const
 	FInputRayHit RayHit;
 	if (TestShouldRespondToMouseWheelFunc.IsBound())
 	{
-		RayHit = TestShouldRespondToMouseWheelFunc.Execute(CurrentPos, BehaviorHost->GetActiveModifiers());
+		RayHit = TestShouldRespondToMouseWheelFunc.Execute(CurrentPos);
 	}
 	return RayHit;
 }
@@ -60,7 +60,7 @@ void UScriptableToolMouseWheelBehavior::OnMouseWheelScrollUp(const FInputDeviceR
 
 void UScriptableToolMouseWheelBehavior::OnMouseWheelScrollDown(const FInputDeviceRay& CurrentPos)
 {
-	OnMouseWheelScrollUpFunc.ExecuteIfBound(CurrentPos, BehaviorHost->GetActiveModifiers());
+	OnMouseWheelScrollDownFunc.ExecuteIfBound(CurrentPos, BehaviorHost->GetActiveModifiers());
 }
 
 void UScriptableToolMouseWheelBehavior::OnUpdateModifierState(int ModifierID, bool bIsOn)
diff --git a/Engine/Plugins/Runtime/ScriptableToolsFramework/Source/ScriptableToolsFramework/Private/Behaviors/ScriptableToolMultiClickSequenceBehavior.cpp b/Engine/Plugins/Runtime/ScriptableToolsFramework/Source/ScriptableToolsFramework/Private/Behaviors/ScriptableToolMultiClickSequenceBehavior.cpp
--- a/Engine/Plugins/Runtime/ScriptableToolsFramework/Source/ScriptableToolsFramework/Private/Behaviors/ScriptableToolMultiClickSequenceBehavior.cpp
+++ b/Engine/Plugins/Runtime/ScriptableToolsFramework/Source/ScriptableToolsFramework/Private/Behaviors/ScriptableToolMultiClickSequenceBehavior.cpp
@@ -82,7 +82,7 @@ bool UScriptableToolClickSequenceBehavior::CanBeginClickSequence(const FInputDev
 	bool bShouldBegin = false;
 	if (CanBeginClickSequenceFunc.IsBound())
 	{
-		bShouldBegin = CanBeginClickSequenceFunc.Execute(ClickPos, BehaviorHost->GetActiveModifiers(), MouseButton);
+		bShouldBegin = CanBeginClickSequenceFunc.Execute(ClickPos, MouseButton);
 	}
 	return bShouldBegin;
 }
diff --git a/Engine/Plugins/Runtime/ScriptableToolsFramework/Source/ScriptableToolsFramework/Private/Behaviors/ScriptableToolSingleClickBehavior.cpp b/Engine/Plugins/Runtime/ScriptableToolsFramework/Source/ScriptableToolsFramework/Private/Behaviors/ScriptableToolSingleClickBehavior.cpp
--- a/Engine/Plugins/Runtime/ScriptableToolsFramework/Source/ScriptableToolsFramework/Private/Behaviors/ScriptableToolSingleClickBehavior.cpp
+++ b/Engine/Plugins/Runtime/ScriptableToolsFramework/Source/ScriptableToolsFramework/Private/Behaviors/ScriptableToolSingleClickBehavior.cpp
@@ -71,7 +71,7 @@ FInputRayHit UScriptableToolSingleClickBehavior::IsHitByClick(const FInputDevice
 	FInputRayHit RayHit;
 	if (TestIfHitByClickFunc.IsBound())
 	{
-		RayHit = TestIfHitByClickFunc.Execute(ClickPos, BehaviorHost->GetActiveModifiers(), MouseButton);
+		RayHit = TestIfHitByClickFunc.Execute(ClickPos, MouseButton);
 	}
 	return RayHit;
 }
diff --git a/Engine/Plugins/Runtime/ScriptableToolsFramework/Source/ScriptableToolsFramework/Private/Behaviors/ScriptableToolSingleClickOrDragBehavior.cpp b/Engine/Plugins/Runtime/ScriptableToolsFramework/Source/ScriptableToolsFramework/Private/Behaviors/ScriptableToolSingleClickOrDragBehavior.cpp
--- a/Engine/Plugins/Runtime/ScriptableToolsFramework/Source/ScriptableToolsFramework/Private/Behaviors/ScriptableToolSingleClickOrDragBehavior.cpp
+++ b/Engine/Plugins/Runtime/ScriptableToolsFramework/Source/ScriptableToolsFramework/Private/Behaviors/ScriptableToolSingleClickOrDragBehavior.cpp
@@ -82,7 +82,7 @@ FInputRayHit UScriptableToolSingleClickOrDragBehavior::IsHitByClick(const FInput
 	FInputRayHit RayHit;
 	if (TestIfHitByClickFunc.IsBound())
 	{
-		RayHit = TestIfHitByClickFunc.Execute(ClickPos, BehaviorHost->GetActiveModifiers(), MouseButton);
+		RayHit = TestIfHitByClickFunc.Execute(ClickPos, MouseButton);
 	}
 	return RayHit;
 }
diff --git a/Engine/Plugins/Runtime/ScriptableToolsFramework/Source/ScriptableToolsFramework/Public/BaseTools/ScriptableModularBehaviorTool.h b/Engine/Plugins/Runtime/ScriptableToolsFramework/Source/ScriptableToolsFramework/Public/BaseTools/ScriptableModularBehaviorTool.h
--- a/Engine/Plugins/Runtime/ScriptableToolsFramework/Source/ScriptableToolsFramework/Public/BaseTools/ScriptableModularBehaviorTool.h
+++ b/Engine/Plugins/Runtime/ScriptableToolsFramework/Source/ScriptableToolsFramework/Public/BaseTools/ScriptableModularBehaviorTool.h
@@ -46,9 +46,9 @@ public:
 	*/
 	UFUNCTION(BlueprintCallable, Category = "ScriptableTool|Input", meta = (AdvancedDisplay = "CaptureCheck, CapturePriority, MouseButton, bHitTestOnRelease"))
 	void AddSingleClickBehavior(
-		const FTestIfHitByClickDelegate& IfHitByClick,
-		const FOnHitByClickDelegate& OnHitByClick,
-		const FMouseBehaviorModiferCheckDelegate& CaptureCheck,
+		const FTestIfHitByClickDelegate IfHitByClick,
+		const FOnHitByClickDelegate OnHitByClick,
+		const FMouseBehaviorModiferCheckDelegate CaptureCheck,
 		int CapturePriority = 100,
 		EScriptableToolMouseButton MouseButton = EScriptableToolMouseButton::LeftButton,
 		bool bHitTestOnRelease = true
@@ -75,12 +75,12 @@ public:
 	*/
 	UFUNCTION(BlueprintCallable, Category = "ScriptableTool|Input", meta = (AdvancedDisplay = "CaptureCheck, CapturePriority, MouseButton, bUpdateModifiersDuringDrag"))
 	void AddClickDragBehavior(		
-		const FTestCanBeginClickDragSequenceDelegate& CanBeginClickDragSequence,
-		const FOnClickPressDelegate& OnClickPress,
-		const FOnClickDragDelegate& OnClickDrag,
-		const FOnClickReleaseDelegate& OnClickRelease,
-		const FOnTerminateDragSequenceDelegate& OnTerminateDragSequence,
-		const FMouseBehaviorModiferCheckDelegate& CaptureCheck,
+		const FTestCanBeginClickDragSequenceDelegate CanBeginClickDragSequence,
+		const FOnClickPressDelegate OnClickPress,
+		const FOnClickDragDelegate OnClickDrag,
+		const FOnClickReleaseDelegate OnClickRelease,
+		const FOnTerminateDragSequenceDelegate OnTerminateDragSequence,
+		const FMouseBehaviorModiferCheckDelegate CaptureCheck,
 		int CapturePriority = 100,
 		EScriptableToolMouseButton MouseButton = EScriptableToolMouseButton::LeftButton,
 		bool bUpdateModifiersDuringDrag = false
diff --git a/Engine/Plugins/Runtime/ScriptableToolsFramework/Source/ScriptableToolsFramework/Public/BaseTools/ScriptableSingleClickTool.h b/Engine/Plugins/Runtime/ScriptableToolsFramework/Source/ScriptableToolsFramework/Public/BaseTools/ScriptableSingleClickTool.h
--- a/Engine/Plugins/Runtime/ScriptableToolsFramework/Source/ScriptableToolsFramework/Public/BaseTools/ScriptableSingleClickTool.h
+++ b/Engine/Plugins/Runtime/ScriptableToolsFramework/Source/ScriptableToolsFramework/Public/BaseTools/ScriptableSingleClickTool.h
@@ -24,7 +24,7 @@ class UMouseHoverBehavior;
  * Note: Currently these interactions are only supported for mouse devices, or devices remapped to 
  * be a mouse by the UInputRouter. Also currently only the left-mouse-button is supported.
  */
-UCLASS()
+UCLASS(Hidden)
 class SCRIPTABLETOOLSFRAMEWORK_API UScriptableSingleClickTool : public UScriptableInteractiveTool, public IClickBehaviorTarget, public IHoverBehaviorTarget
 {
 	GENERATED_BODY()
diff --git a/Engine/Plugins/Runtime/ScriptableToolsFramework/Source/ScriptableToolsFramework/Public/Behaviors/ScriptableToolBehaviorDelegates.h b/Engine/Plugins/Runtime/ScriptableToolsFramework/Source/ScriptableToolsFramework/Public/Behaviors/ScriptableToolBehaviorDelegates.h
--- a/Engine/Plugins/Runtime/ScriptableToolsFramework/Source/ScriptableToolsFramework/Public/Behaviors/ScriptableToolBehaviorDelegates.h
+++ b/Engine/Plugins/Runtime/ScriptableToolsFramework/Source/ScriptableToolsFramework/Public/Behaviors/ScriptableToolBehaviorDelegates.h
@@ -17,7 +17,7 @@ enum class EScriptableToolMouseButton : uint8
 
 DECLARE_DYNAMIC_DELEGATE_RetVal_OneParam(bool, FMouseBehaviorModiferCheckDelegate, const FInputDeviceState&, InputDeviceState);
 
-DECLARE_DYNAMIC_DELEGATE_RetVal_ThreeParams(FInputRayHit, FTestIfHitByClickDelegate, FInputDeviceRay, ClickPos, FScriptableToolModifierStates, Modifiers, EScriptableToolMouseButton, MouseButton);
+DECLARE_DYNAMIC_DELEGATE_RetVal_TwoParams(FInputRayHit, FTestIfHitByClickDelegate, FInputDeviceRay, ClickPos, EScriptableToolMouseButton, MouseButton);
 DECLARE_DYNAMIC_DELEGATE_ThreeParams(FOnHitByClickDelegate, FInputDeviceRay, ClickPos, FScriptableToolModifierStates, Modifiers, EScriptableToolMouseButton, MouseButton);
 
 DECLARE_DYNAMIC_DELEGATE_RetVal_ThreeParams(FInputRayHit, FTestCanBeginClickDragSequenceDelegate, FInputDeviceRay, PressPos, FScriptableToolModifierStates, Modifiers, EScriptableToolMouseButton, MouseButton);
@@ -26,12 +26,12 @@ DECLARE_DYNAMIC_DELEGATE_ThreeParams(FOnClickDragDelegate, FInputDeviceRay, Drag
 DECLARE_DYNAMIC_DELEGATE_ThreeParams(FOnClickReleaseDelegate, FInputDeviceRay, ReleasePos, FScriptableToolModifierStates, Modifiers, EScriptableToolMouseButton, MouseButton);
 DECLARE_DYNAMIC_DELEGATE_TwoParams(FOnTerminateDragSequenceDelegate, FScriptableToolModifierStates, Modifiers, EScriptableToolMouseButton, MouseButton);
 
-DECLARE_DYNAMIC_DELEGATE_RetVal_TwoParams(FInputRayHit, FTestShouldRespondToMouseWheelDelegate, FInputDeviceRay, CurrentPos, FScriptableToolModifierStates, Modifiers);
+DECLARE_DYNAMIC_DELEGATE_RetVal_OneParam(FInputRayHit, FTestShouldRespondToMouseWheelDelegate, FInputDeviceRay, CurrentPos);
 DECLARE_DYNAMIC_DELEGATE_TwoParams(FOnMouseWheelScrollUpDelegate, FInputDeviceRay, CurrentPos, FScriptableToolModifierStates, Modifiers);
 DECLARE_DYNAMIC_DELEGATE_TwoParams(FOnMouseWheelScrollDownDelegate, FInputDeviceRay, CurrentPos, FScriptableToolModifierStates, Modifiers);
 
 DECLARE_DYNAMIC_DELEGATE_ThreeParams(FOnBeginSequencePreviewDelegate, FInputDeviceRay, ClickPos, FScriptableToolModifierStates, Modifiers, EScriptableToolMouseButton, MouseButton);
-DECLARE_DYNAMIC_DELEGATE_RetVal_ThreeParams(bool, FCanBeginClickSequenceDelegate, FInputDeviceRay, ClickPos, FScriptableToolModifierStates, Modifiers, EScriptableToolMouseButton, MouseButton);
+DECLARE_DYNAMIC_DELEGATE_RetVal_TwoParams(bool, FCanBeginClickSequenceDelegate, FInputDeviceRay, ClickPos, EScriptableToolMouseButton, MouseButton);
 DECLARE_DYNAMIC_DELEGATE_ThreeParams(FOnBeginClickSequenceDelegate, FInputDeviceRay, ClickPos, FScriptableToolModifierStates, Modifiers, EScriptableToolMouseButton, MouseButton);
 DECLARE_DYNAMIC_DELEGATE_ThreeParams(FOnNextSequencePreviewDelegate, FInputDeviceRay, ClickPos, FScriptableToolModifierStates, Modifiers, EScriptableToolMouseButton, MouseButton);
 DECLARE_DYNAMIC_DELEGATE_RetVal_ThreeParams(bool, FOnNextSequenceClickDelegate, FInputDeviceRay, ClickPos, FScriptableToolModifierStates, Modifiers, EScriptableToolMouseButton, MouseButton);
