ScriptableTools: Adding support for customizable ToolBuilders for Scriptable Tools.

In order to support behaviors in Scriptable Tools such as startup requirements and the Interactive Tool Framework's concept of Tool Targets, support for ToolBuilders has been added to the Scriptable Tool base class. This support comes in three forms:

* None - A built in "standard" tool builder that offers no customization and no pre-startup filtering.
* Tool Target - A customizable tool builder that allows users to construct tool target requirements and target counts. Sets tool targets on the resulting tool instance for future expansion of tool target support with specific modes.
* Custom - A customizable tool builder that provides users access to the scene selection for pre-startup requirements and further pre-setup access to the tool instance, representing the BuildTool portion of the tool builder API.

diff --git a/Engine/Plugins/Editor/ScriptableToolsEditorMode/Source/ScriptableToolsEditorMode/Private/ScriptableToolsEditorMode.cpp b/Engine/Plugins/Editor/ScriptableToolsEditorMode/Source/ScriptableToolsEditorMode/Private/ScriptableToolsEditorMode.cpp
--- a/Engine/Plugins/Editor/ScriptableToolsEditorMode/Source/ScriptableToolsEditorMode/Private/ScriptableToolsEditorMode.cpp
+++ b/Engine/Plugins/Editor/ScriptableToolsEditorMode/Source/ScriptableToolsEditorMode/Private/ScriptableToolsEditorMode.cpp
@@ -22,6 +22,11 @@
 #include "InteractiveToolQueryInterfaces.h" // IInteractiveToolExclusiveToolAPI
 #include "ToolContextInterfaces.h"
 
+#include "ToolTargetManager.h"
+#include "ToolTargets/StaticMeshComponentToolTarget.h"
+#include "ToolTargets/VolumeComponentToolTarget.h"
+#include "ToolTargets/DynamicMeshComponentToolTarget.h"
+#include "ToolTargets/SkeletalMeshComponentToolTarget.h"
 
 #define LOCTEXT_NAMESPACE "UScriptableToolsEditorMode"
 
@@ -132,6 +137,13 @@ void UScriptableToolsEditorMode::Enter()
 	// listen to post-build
 	GetToolManager()->OnToolPostBuild.AddUObject(this, &UScriptableToolsEditorMode::OnToolPostBuild);
 
+	// Register builders for tool targets that the mode uses.
+	// TODO: We're not actually suporting modeling mode tool targets on scriptable tools, but the infrastructure to test for selected
+	// objects uses the ToolTargetFactories, so we're including these here. We probably need a more generic way to accomplish this.
+	GetInteractiveToolsContext()->TargetManager->AddTargetFactory(NewObject<UStaticMeshComponentToolTargetFactory>(GetToolManager()));
+	GetInteractiveToolsContext()->TargetManager->AddTargetFactory(NewObject<UVolumeComponentToolTargetFactory>(GetToolManager()));
+	GetInteractiveToolsContext()->TargetManager->AddTargetFactory(NewObject<UDynamicMeshComponentToolTargetFactory>(GetToolManager()));
+
 	//// forward shutdown requests
 	//GetToolManager()->OnToolShutdownRequest.BindLambda([this](UInteractiveToolManager*, UInteractiveTool* Tool, EToolShutdownType ShutdownType)
 	//{
diff --git a/Engine/Plugins/Editor/ScriptableToolsEditorMode/Source/ScriptableToolsEditorMode/Private/ScriptableToolsEditorModeToolkit.cpp b/Engine/Plugins/Editor/ScriptableToolsEditorMode/Source/ScriptableToolsEditorMode/Private/ScriptableToolsEditorModeToolkit.cpp
--- a/Engine/Plugins/Editor/ScriptableToolsEditorMode/Source/ScriptableToolsEditorMode/Private/ScriptableToolsEditorModeToolkit.cpp
+++ b/Engine/Plugins/Editor/ScriptableToolsEditorMode/Source/ScriptableToolsEditorMode/Private/ScriptableToolsEditorModeToolkit.cpp
@@ -371,7 +371,6 @@ void FScriptableToolsEditorModeToolkit::BuildToolPalette(FName PaletteIndex, cla
 		FExecuteAction::CreateLambda([this, ToolClass, ToolIdentifier, ToolManager]()
 		{
 			UScriptableInteractiveTool* ToolCDO = Cast<UScriptableInteractiveTool>(ToolClass->GetDefaultObject());
-			//UE_LOG(LogTemp, Warning, TEXT("STARTING TOOL [%s] (Class/Identifier %s)"), *ToolCDO->ToolName.ToString(), *ToolIdentifier);
 			if (ToolManager->SelectActiveToolType(EToolSide::Mouse, ToolIdentifier))
 			{
 				if (ToolManager->CanActivateTool(EToolSide::Mouse, ToolIdentifier)) 
@@ -385,6 +384,19 @@ void FScriptableToolsEditorModeToolkit::BuildToolPalette(FName PaletteIndex, cla
 				UE_LOG(LogTemp, Warning, TEXT("FAILED TO SET ACTIVE TOOL TYPE!"));
 			}
 
+		}),
+		FCanExecuteAction::CreateLambda([this, ToolClass, ToolIdentifier, ToolManager]()
+		{
+			UScriptableInteractiveTool* ToolCDO = Cast<UScriptableInteractiveTool>(ToolClass->GetDefaultObject());
+			if (ToolManager->SelectActiveToolType(EToolSide::Mouse, ToolIdentifier))
+			{
+				return ToolManager->CanActivateTool(EToolSide::Mouse, ToolIdentifier);
+			}
+			else
+			{
+				UE_LOG(LogTemp, Warning, TEXT("FAILED TO SET ACTIVE TOOL TYPE!"));
+				return false;
+			}
 		}));
 
 		ActionsHack.Add(NewAction);
diff --git a/Engine/Plugins/Editor/ScriptableToolsEditorMode/Source/ScriptableToolsEditorMode/ScriptableToolsEditorMode.Build.cs b/Engine/Plugins/Editor/ScriptableToolsEditorMode/Source/ScriptableToolsEditorMode/ScriptableToolsEditorMode.Build.cs
--- a/Engine/Plugins/Editor/ScriptableToolsEditorMode/Source/ScriptableToolsEditorMode/ScriptableToolsEditorMode.Build.cs
+++ b/Engine/Plugins/Editor/ScriptableToolsEditorMode/Source/ScriptableToolsEditorMode/ScriptableToolsEditorMode.Build.cs
@@ -48,6 +48,7 @@ public class ScriptableToolsEditorMode : ModuleRules
 				"EditorInteractiveToolsFramework",
 				"ScriptableToolsFramework",
 				"EditorScriptableToolsFramework",
+				"ModelingComponentsEditorOnly",
 
 				"ModelingComponents"
 			}
diff --git a/Engine/Plugins/Runtime/ScriptableToolsFramework/Source/ScriptableToolsFramework/Private/ScriptableInteractiveTool.cpp b/Engine/Plugins/Runtime/ScriptableToolsFramework/Source/ScriptableToolsFramework/Private/ScriptableInteractiveTool.cpp
--- a/Engine/Plugins/Runtime/ScriptableToolsFramework/Source/ScriptableToolsFramework/Private/ScriptableInteractiveTool.cpp
+++ b/Engine/Plugins/Runtime/ScriptableToolsFramework/Source/ScriptableToolsFramework/Private/ScriptableInteractiveTool.cpp
@@ -267,6 +267,45 @@ void UScriptableInteractiveTool::PostInitProperties()
 }
 
 
+UBaseScriptableToolBuilder* UScriptableInteractiveTool::GetNewCustomToolBuilderInstance(UObject* Outer)
+{
+	switch (ToolStartupRequirements)
+	{
+	case EScriptableToolStartupRequirements::None:
+		return NewObject<UBaseScriptableToolBuilder>(Outer);
+	case EScriptableToolStartupRequirements::ToolTarget:
+	{
+		if (!ToolTargetToolBuilderClass->IsValidLowLevelFast())
+		{
+			return NewObject<UBaseScriptableToolBuilder>(Outer);
+		}
+
+		TObjectPtr<UCustomScriptableToolBuilderContainer> BuilderWrapper = NewObject<UCustomScriptableToolBuilderContainer>(Outer);
+		TObjectPtr<UCustomScriptableToolBuilderComponentBase> InterfacePointer = NewObject<UCustomScriptableToolBuilderComponentBase>(Outer, ToolTargetToolBuilderClass.Get());
+		Cast<UToolTargetScriptableToolBuilder>(InterfacePointer)->Initialize();
+		BuilderWrapper->Initialize(InterfacePointer);
+
+		return BuilderWrapper;
+	}	
+	case EScriptableToolStartupRequirements::Custom:
+	{
+		if (!CustomToolBuilderClass->IsValidLowLevelFast())
+		{
+			return NewObject<UBaseScriptableToolBuilder>(Outer);
+		}
+
+		TObjectPtr<UCustomScriptableToolBuilderContainer> BuilderWrapper = NewObject<UCustomScriptableToolBuilderContainer>(Outer);
+		TObjectPtr<UCustomScriptableToolBuilderComponentBase> InterfacePointer = NewObject<UCustomScriptableToolBuilderComponentBase>(Outer, CustomToolBuilderClass.Get());
+		BuilderWrapper->Initialize(InterfacePointer);
+
+		return BuilderWrapper;
+	}
+	default:
+		ensure(false);
+		return NewObject<UBaseScriptableToolBuilder>(Outer);
+	}
+}
+
 
 void UScriptableInteractiveTool::SetTargetWorld(UWorld* World)
 {
@@ -1089,6 +1128,18 @@ void UScriptableInteractiveTool::ClearUserMessages(bool bNotifications, bool bWa
 }
 
 
+void UScriptableInteractiveTool::SetTargets(TArray<TObjectPtr<UToolTarget>> TargetsIn)
+{
+	Targets = TargetsIn;
+}
+
+TArray<UToolTarget*> UScriptableInteractiveTool::GetToolTargets() const
+{
+	return Targets;
+}
+
+
+
 
 
 
diff --git a/Engine/Plugins/Runtime/ScriptableToolsFramework/Source/ScriptableToolsFramework/Private/ScriptableToolBuilder.cpp b/Engine/Plugins/Runtime/ScriptableToolsFramework/Source/ScriptableToolsFramework/Private/ScriptableToolBuilder.cpp
--- a/Engine/Plugins/Runtime/ScriptableToolsFramework/Source/ScriptableToolsFramework/Private/ScriptableToolBuilder.cpp
+++ b/Engine/Plugins/Runtime/ScriptableToolsFramework/Source/ScriptableToolsFramework/Private/ScriptableToolBuilder.cpp
@@ -5,6 +5,11 @@
 #include "ScriptableInteractiveTool.h"
 #include "ToolContextInterfaces.h"
 
+#include "ToolTargetManager.h"
+#include "TargetInterfaces/MaterialProvider.h"
+#include "TargetInterfaces/MeshDescriptionCommitter.h"
+#include "TargetInterfaces/MeshDescriptionProvider.h"
+#include "TargetInterfaces/PrimitiveComponentBackedTarget.h"
 
 bool UBaseScriptableToolBuilder::CanBuildTool(const FToolBuilderState& SceneState) const
 {
@@ -24,3 +29,96 @@ UInteractiveTool* UBaseScriptableToolBuilder::BuildTool(const FToolBuilderState&
 	return NewTool;
 }
 
+void UCustomScriptableToolBuilderContainer::Initialize(TObjectPtr<UCustomScriptableToolBuilderComponentBase> BuilderInstanceIn)
+{
+	BuilderInstance = BuilderInstanceIn;
+}
+
+bool UCustomScriptableToolBuilderContainer::CanBuildTool(const FToolBuilderState& SceneState) const
+{
+	bool bCanBuild = Super::CanBuildTool(SceneState);
+
+	ICustomScriptableToolBuilderBaseInterface* Builder = Cast<ICustomScriptableToolBuilderBaseInterface>(BuilderInstance);
+	bCanBuild = bCanBuild && Builder->CanBuildTool(SceneState);
+
+	return bCanBuild;
+
+}
+
+UInteractiveTool* UCustomScriptableToolBuilderContainer::BuildTool(const FToolBuilderState& SceneState) const
+{
+	UInteractiveTool* NewToolObj = UBaseScriptableToolBuilder::BuildTool(SceneState);
+
+	ICustomScriptableToolBuilderBaseInterface* Builder = Cast<ICustomScriptableToolBuilderBaseInterface>(BuilderInstance);
+	Builder->SetupTool(SceneState, NewToolObj);
+
+	return NewToolObj;
+}
+
+UScriptableToolTargetRequirements*
+UScriptableToolTargetRequirements::BuildToolTargetRequirements(TArray<UClass*> RequirementInterfaces)
+{
+	TObjectPtr<UScriptableToolTargetRequirements> ScriptableToolRequirements = NewObject<UScriptableToolTargetRequirements>();
+
+	for (UClass*& ClassPtr : RequirementInterfaces)
+	{
+		ScriptableToolRequirements->Requirements.Add(ClassPtr);
+	}
+
+	return ScriptableToolRequirements;
+}
+
+
+bool UCustomScriptableToolBuilder::CanBuildTool(const FToolBuilderState& SceneState) const
+{
+	return OnCanBuildTool(SceneState.SelectedActors, SceneState.SelectedComponents);
+}
+
+bool UCustomScriptableToolBuilder::OnCanBuildTool_Implementation(const TArray<AActor*>& SelectedActors, const TArray<UActorComponent*>& SelectedComponents) const
+{
+	return true;
+}
+
+void UCustomScriptableToolBuilder::SetupTool(const FToolBuilderState& SceneState, UInteractiveTool* Tool) const
+{
+	UScriptableInteractiveTool* NewTool = Cast<UScriptableInteractiveTool>(Tool);
+
+	OnSetupTool(NewTool, SceneState.SelectedActors, SceneState.SelectedComponents);
+}
+
+void UCustomScriptableToolBuilder::OnSetupTool_Implementation(UScriptableInteractiveTool* Tool, const TArray<AActor*>& SelectedActors, const TArray<UActorComponent*>& SelectedComponents) const
+{
+
+}
+
+
+void UToolTargetScriptableToolBuilder::Initialize()
+{
+	Requirements = GetToolTargetRequirements();
+}
+
+UScriptableToolTargetRequirements* UToolTargetScriptableToolBuilder::GetToolTargetRequirements_Implementation() const
+{
+	return NewObject<UScriptableToolTargetRequirements>();
+}
+
+bool UToolTargetScriptableToolBuilder::CanBuildTool(const FToolBuilderState& SceneState) const
+{
+	int MatchingTargetCount = SceneState.TargetManager->CountSelectedAndTargetable(SceneState, Requirements->GetRequirements());
+	return MatchingTargetCount >= Requirements->MinMatchingTargets && MatchingTargetCount <= Requirements->MaxMatchingTargets;
+}
+
+void UToolTargetScriptableToolBuilder::SetupTool(const FToolBuilderState& SceneState, UInteractiveTool* Tool) const
+{
+	UScriptableInteractiveTool* NewTool = Cast<UScriptableInteractiveTool>(Tool);
+	TArray<TObjectPtr<UToolTarget>> Targets = SceneState.TargetManager->BuildAllSelectedTargetable(SceneState, Requirements->GetRequirements());
+	const int NumTargets = Targets.Num();
+	Targets.SetNum(FMath::Min(NumTargets, Requirements->MaxMatchingTargets));
+	NewTool->SetTargets(Targets);
+	OnSetupTool(NewTool);
+}
+
+void UToolTargetScriptableToolBuilder::OnSetupTool_Implementation(UScriptableInteractiveTool* Tool) const
+{
+
+}
\ No newline at end of file
diff --git a/Engine/Plugins/Runtime/ScriptableToolsFramework/Source/ScriptableToolsFramework/Public/ScriptableInteractiveTool.h b/Engine/Plugins/Runtime/ScriptableToolsFramework/Source/ScriptableToolsFramework/Public/ScriptableInteractiveTool.h
--- a/Engine/Plugins/Runtime/ScriptableToolsFramework/Source/ScriptableToolsFramework/Public/ScriptableInteractiveTool.h
+++ b/Engine/Plugins/Runtime/ScriptableToolsFramework/Source/ScriptableToolsFramework/Public/ScriptableInteractiveTool.h
@@ -5,6 +5,7 @@
 #include "InteractiveTool.h"
 #include "InteractiveGizmo.h"
 #include "Kismet/BlueprintFunctionLibrary.h"
+#include "ScriptableToolBuilder.h"
 #include "ScriptableInteractiveTool.generated.h"
 
 class UWorld;
@@ -122,6 +123,16 @@ enum class EScriptableToolGizmoScale : uint8
 };
 ENUM_CLASS_FLAGS(EScriptableToolGizmoScale);
 
+UENUM(BlueprintType)
+enum class EScriptableToolStartupRequirements : uint8
+{
+	/** No startup requirements needed. Tool can run any time. */
+	None,
+	/** A custom tool builder blueprint class that is configured with tool target requirements to filter selected objects. */
+	ToolTarget,
+	/** A custom tool builder blueprint class is provided to determine if the tool can start. Caution: OnCanBuildTool is run every tick, and may slow down editor performance. */
+	Custom
+};
 
 /**
  * FScriptableToolGizmoOptions is a configuration struct passed to the CreateTRSGizmo function
@@ -318,6 +329,23 @@ public:
 	UPROPERTY(BlueprintReadWrite, EditAnywhere, Category = "Scriptable Tool Settings", meta=(DisplayName="Shutdown Type"))
 	EScriptableToolShutdownType ToolShutdownType = EScriptableToolShutdownType::Complete;
 
+	UPROPERTY(BlueprintReadWrite, EditAnywhere, Category = "Scriptable Tool Settings", meta = (DisplayName = "Tool Startup Requirements"))
+	EScriptableToolStartupRequirements ToolStartupRequirements = EScriptableToolStartupRequirements::None;
+
+	UPROPERTY(BlueprintReadWrite, EditAnywhere, Category = "Scriptable Tool Settings", meta = (DisplayName = "Tool Builder Class",
+		                                                                                       EditCondition = "ToolStartupRequirements==EScriptableToolStartupRequirements::Custom",
+																							   EditConditionHides,
+																					           MustImplement = "/Script/ScriptableToolsFramework.CustomScriptableToolBuilderBaseInterface"))
+	TSubclassOf<UCustomScriptableToolBuilder> CustomToolBuilderClass = nullptr;
+
+	UPROPERTY(BlueprintReadWrite, EditAnywhere, Category = "Scriptable Tool Settings", meta = (DisplayName = "Tool Builder Class",
+		                                                                                       EditCondition = "ToolStartupRequirements==EScriptableToolStartupRequirements::ToolTarget",
+		                                                                                       EditConditionHides,
+		                                                                                       MustImplement="/Script/ScriptableToolsFramework.CustomScriptableToolBuilderBaseInterface"))
+	TSubclassOf<UToolTargetScriptableToolBuilder> ToolTargetToolBuilderClass = nullptr;
+
+
+
 	/**
 	 * Implement OnScriptSetup to do initial setup/configuration of the Tool, such as adding
 	 * Property Sets, creating Gizmos, etc
@@ -385,7 +413,7 @@ protected:
 
 
 	// return instance of custom tool builder. Should only be called on CDO.
-	virtual UBaseScriptableToolBuilder* GetNewCustomToolBuilderInstance(UObject* Outer) { return nullptr; }
+	virtual UBaseScriptableToolBuilder* GetNewCustomToolBuilderInstance(UObject* Outer);
 	friend class UScriptableToolSet;
 
 
@@ -757,6 +785,23 @@ public:
 public:
 
 	// TODO: hotkey API
+
+
+
+public:
+
+	// Tool Targets API
+
+	void SetTargets(TArray<TObjectPtr<UToolTarget>> TargetsIn);
+
+	UFUNCTION(BlueprintCallable, Category = "ScriptableTool|ToolTargets")
+	UPARAM(DisplayName = "Targets") TArray<UToolTarget*> GetToolTargets() const;
+
+protected:
+
+	UPROPERTY(Transient, DuplicateTransient, NonTransactional, SkipSerialization)
+	TArray<TObjectPtr<UToolTarget>> Targets;
+
 };
 
 
diff --git a/Engine/Plugins/Runtime/ScriptableToolsFramework/Source/ScriptableToolsFramework/Public/ScriptableToolBuilder.h b/Engine/Plugins/Runtime/ScriptableToolsFramework/Source/ScriptableToolsFramework/Public/ScriptableToolBuilder.h
--- a/Engine/Plugins/Runtime/ScriptableToolsFramework/Source/ScriptableToolsFramework/Public/ScriptableToolBuilder.h
+++ b/Engine/Plugins/Runtime/ScriptableToolsFramework/Source/ScriptableToolsFramework/Public/ScriptableToolBuilder.h
@@ -3,10 +3,16 @@
 #pragma once
 
 #include "InteractiveToolBuilder.h"
-#include "ScriptableInteractiveTool.h"
 #include "ScriptableToolBuilder.generated.h"
 
 
+namespace ScriptableToolBuilderHelpers
+{
+
+}
+
+class UScriptableInteractiveTool;
+
 /**
  * UBaseScriptableToolBuilder is a trivial base UInteractiveToolBuilder for any UScriptableInteractiveTool subclass.
  * CanBuildTool will return true as long as the ToolClass is a valid UClass.
@@ -23,4 +29,134 @@ public:
 	virtual UInteractiveTool* BuildTool(const FToolBuilderState& SceneState) const override;
 };
 
+UINTERFACE(MinimalAPI, BlueprintType)
+class UCustomScriptableToolBuilderBaseInterface : public UInterface
+{
+	GENERATED_BODY()
+};
+
+class ICustomScriptableToolBuilderBaseInterface
+{
+	GENERATED_BODY()
+
+public:
+
+	virtual bool CanBuildTool(const FToolBuilderState& SceneState) const = 0;
+	virtual void SetupTool(const FToolBuilderState& SceneState, UInteractiveTool* Tool) const = 0;
+};
+
+UCLASS()
+class SCRIPTABLETOOLSFRAMEWORK_API UCustomScriptableToolBuilderComponentBase : public UObject
+{
+	GENERATED_BODY()
+
+public:
+
+};
+
+
+UCLASS(Transient, Blueprintable, Hidden)
+class SCRIPTABLETOOLSFRAMEWORK_API UCustomScriptableToolBuilderContainer : public UBaseScriptableToolBuilder
+{
+	GENERATED_BODY()
+
+
+public:
+
+	void Initialize(TObjectPtr<UCustomScriptableToolBuilderComponentBase> BuilderInstanceIn);
+
+	virtual bool CanBuildTool(const FToolBuilderState& SceneState) const override;
+	virtual UInteractiveTool* BuildTool(const FToolBuilderState& SceneState) const override;
+
+private:
+
+	UPROPERTY()
+	TObjectPtr<UCustomScriptableToolBuilderComponentBase> BuilderInstance;
+};
+
+/*
+*
+*   Tool Builders for custom builder logic
+*
+*/
+
+UCLASS(Transient, Blueprintable, Abstract)
+class SCRIPTABLETOOLSFRAMEWORK_API UCustomScriptableToolBuilder : public UCustomScriptableToolBuilderComponentBase, public ICustomScriptableToolBuilderBaseInterface
+{
+	GENERATED_BODY()
+
+public:
+
+	UFUNCTION(BlueprintNativeEvent, Category = "ScriptableToolBuilder|Events")
+	bool OnCanBuildTool(const TArray<AActor*>& SelectedActors, const TArray<UActorComponent*>& SelectedComponents) const;
+
+	bool OnCanBuildTool_Implementation(const TArray<AActor*>& SelectedActors, const TArray<UActorComponent*>& SelectedComponents) const;
+
+	UFUNCTION(BlueprintNativeEvent, Category = "ScriptableToolBuilder|Events")
+	void OnSetupTool(UScriptableInteractiveTool* Tool, const TArray<AActor*>& SelectedActors, const TArray<UActorComponent*>& SelectedComponents) const;
+
+	void OnSetupTool_Implementation(UScriptableInteractiveTool* Tool, const TArray<AActor*>& SelectedActors, const TArray<UActorComponent*>& SelectedComponents) const;
+
+	virtual bool CanBuildTool(const FToolBuilderState& SceneState) const override;
+	virtual void SetupTool(const FToolBuilderState& SceneState, UInteractiveTool* Tool) const override;
+};
+
+
+/*
+*
+*   Tool Builders for Tool Target support
+* 
+*/
+
+UCLASS(Blueprintable)
+class SCRIPTABLETOOLSFRAMEWORK_API UScriptableToolTargetRequirements : public UObject
+{
+	GENERATED_BODY()
+public:
+
+	UFUNCTION(BlueprintCallable, Category = "ScriptableToolBuilder|ToolTargets")
+	static UPARAM(DisplayName = "New Target Requirements") UScriptableToolTargetRequirements*
+	BuildToolTargetRequirements(TArray<UClass*> RequirementInterfaces);
+
+	const FToolTargetTypeRequirements& GetRequirements() const { return Requirements; };
+
+	UPROPERTY(BlueprintReadWrite, EditAnywhere, Category = "Requirements")
+	int MinMatchingTargets = 0;
+
+	UPROPERTY(BlueprintReadWrite, EditAnywhere, Category = "Requirements")
+	int MaxMatchingTargets = 1;
+
+private:
+
+	FToolTargetTypeRequirements Requirements;
+
+};
+
+UCLASS(Transient, Blueprintable, Abstract)
+class UToolTargetScriptableToolBuilder: public UCustomScriptableToolBuilderComponentBase, public ICustomScriptableToolBuilderBaseInterface
+{
+	GENERATED_BODY()
+
+public:
+
+	void Initialize();
+
+	UFUNCTION(BlueprintNativeEvent, Category = "Tool Targets")
+	UScriptableToolTargetRequirements* GetToolTargetRequirements() const;
+
+	virtual UScriptableToolTargetRequirements* GetToolTargetRequirements_Implementation() const;
+
+	UFUNCTION(BlueprintNativeEvent, Category = "ScriptableToolBuilder|Events")
+	void OnSetupTool(UScriptableInteractiveTool* Tool) const;
+
+	void OnSetupTool_Implementation(UScriptableInteractiveTool* Tool) const;
+
+	virtual bool CanBuildTool(const FToolBuilderState& SceneState) const override;
+	virtual void SetupTool(const FToolBuilderState& SceneState, UInteractiveTool* Tool) const override;
+	
+private:
+
+	UPROPERTY()
+	TObjectPtr<UScriptableToolTargetRequirements> Requirements;
 
+};
\ No newline at end of file
