ScriptableTools: Fix missing plugin requirement

diff --git a/Engine/Plugins/Runtime/ScriptableToolsFramework/ScriptableToolsFramework.uplugin b/Engine/Plugins/Runtime/ScriptableToolsFramework/ScriptableToolsFramework.uplugin
--- a/Engine/Plugins/Runtime/ScriptableToolsFramework/ScriptableToolsFramework.uplugin
+++ b/Engine/Plugins/Runtime/ScriptableToolsFramework/ScriptableToolsFramework.uplugin
@@ -10,7 +10,7 @@
 	"DocsURL": "",
 	"MarketplaceURL": "",
 	"SupportURL": "",
-    "EnabledByDefault": false,
+	"EnabledByDefault": false,
 	"CanContainContent": true,
 	"IsBetaVersion": true,
 	"Installed": false,
@@ -24,8 +24,12 @@
 			"Name": "EditorScriptableToolsFramework",
 			"Type": "Editor",
 			"LoadingPhase": "Default"
-		}		
+		}
 	],
 	"Plugins": [
+		{
+			"Name": "MeshModelingToolset",
+			"Enabled": true
+		}
 	]
 }
