ScriptableTools: Add support for preview geomety based point/line/triangle sets to scriptable tools, allowing users to create and manage persisent geometry sets and elements for in tool visualization and temporary geometry creation.

diff --git a/Engine/Plugins/Runtime/ScriptableToolsFramework/Source/ScriptableToolsFramework/Private/Drawing/ScriptableToolLine.cpp b/Engine/Plugins/Runtime/ScriptableToolsFramework/Source/ScriptableToolsFramework/Private/Drawing/ScriptableToolLine.cpp
new file mode 100644
--- /dev/null
+++ b/Engine/Plugins/Runtime/ScriptableToolsFramework/Source/ScriptableToolsFramework/Private/Drawing/ScriptableToolLine.cpp
@@ -0,0 +1,62 @@
+// Copyright Epic Games, Inc. All Rights Reserved.
+
+#include "Drawing/ScriptableToolLine.h"
+
+
+void UScriptableToolLine::SetLineID(int32 LineIDIn)
+{
+	LineID = LineIDIn;
+}
+
+int32 UScriptableToolLine::GetLineID() const
+{
+	return LineID;
+}
+
+bool UScriptableToolLine::IsDirty() const
+{
+	return bIsDirty;
+}
+
+FRenderableLine UScriptableToolLine::GenerateLineDescription()
+{
+	bIsDirty = false;
+	return LineDescription;
+}
+
+void UScriptableToolLine::SetLineStart(FVector LineStartIn)
+{
+	LineDescription.Start = LineStartIn;
+	bIsDirty = true;
+}
+
+void UScriptableToolLine::SetLineEnd(FVector LineEndIn)
+{
+	LineDescription.End = LineEndIn;
+	bIsDirty = true;
+}
+
+void UScriptableToolLine::SetLineEndPoints(FVector LineStartIn, FVector LineEndIn)
+{
+	LineDescription.Start = LineStartIn;
+	LineDescription.End = LineEndIn;
+	bIsDirty = true;
+}
+
+void UScriptableToolLine::SetLineColor(FColor LineColorIn)
+{
+	LineDescription.Color = LineColorIn;
+	bIsDirty = true;
+}
+
+void UScriptableToolLine::SetLineThickness(float LineThicknessIn)
+{
+	LineDescription.Thickness = LineThicknessIn;
+	bIsDirty = true;
+}
+
+void UScriptableToolLine::SetLineDepthBias(float LineDepthBiasIn)
+{
+	LineDescription.DepthBias = LineDepthBiasIn;
+	bIsDirty = true;
+}
\ No newline at end of file
diff --git a/Engine/Plugins/Runtime/ScriptableToolsFramework/Source/ScriptableToolsFramework/Private/Drawing/ScriptableToolLineSet.cpp b/Engine/Plugins/Runtime/ScriptableToolsFramework/Source/ScriptableToolsFramework/Private/Drawing/ScriptableToolLineSet.cpp
new file mode 100644
--- /dev/null
+++ b/Engine/Plugins/Runtime/ScriptableToolsFramework/Source/ScriptableToolsFramework/Private/Drawing/ScriptableToolLineSet.cpp
@@ -0,0 +1,63 @@
+// Copyright Epic Games, Inc. All Rights Reserved.
+
+#include "Drawing/ScriptableToolLineSet.h"
+
+#include "Drawing/ScriptableToolLine.h"
+#include "Drawing/PreviewGeometryActor.h"
+#include "Misc/Guid.h"
+
+void UScriptableToolLineSet::Initialize(TObjectPtr<UPreviewGeometry> PreviewGeometry)
+{
+	FString LineSetID = FGuid::NewGuid().ToString();
+	LineSet = PreviewGeometry->AddLineSet(LineSetID);
+}
+
+void UScriptableToolLineSet::OnTick()
+{
+	for (TObjectPtr<UScriptableToolLine> LineComponent : LineComponents)
+	{
+		if (LineComponent->IsDirty())
+		{
+			int32 LineID = LineComponent->GetLineID();
+			FRenderableLine LineDescription = LineComponent->GenerateLineDescription();
+
+			LineSet->SetLineStart(LineID, LineDescription.Start);
+			LineSet->SetLineEnd(LineID, LineDescription.End);
+			LineSet->SetLineColor(LineID, LineDescription.Color);
+			LineSet->SetLineThickness(LineID, LineDescription.Thickness);
+		}
+	}
+}
+
+UScriptableToolLine* UScriptableToolLineSet::AddLine()
+{
+	LineComponents.Add(NewObject<UScriptableToolLine>(this));
+	UScriptableToolLine* NewLine = LineComponents.Last();
+	NewLine->SetLineID(LineSet->AddLine(NewLine->GenerateLineDescription()));
+	return NewLine;
+}
+
+void UScriptableToolLineSet::RemoveLine(UScriptableToolLine* Line)
+{
+	if (ensure(Line))
+	{
+		LineSet->RemoveLine(Line->GetLineID());
+		LineComponents.Remove(Line);
+	}
+}
+
+void UScriptableToolLineSet::RemoveAllLines()
+{
+	LineSet->Clear();
+	LineComponents.Empty();
+}
+
+void UScriptableToolLineSet::SetAllLinesColor(FColor Color)
+{
+	LineSet->SetAllLinesColor(Color);
+}
+
+void UScriptableToolLineSet::SetAllLinesThickness(float Thickness)
+{
+	LineSet->SetAllLinesThickness(Thickness);
+}
\ No newline at end of file
diff --git a/Engine/Plugins/Runtime/ScriptableToolsFramework/Source/ScriptableToolsFramework/Private/Drawing/ScriptableToolPoint.cpp b/Engine/Plugins/Runtime/ScriptableToolsFramework/Source/ScriptableToolsFramework/Private/Drawing/ScriptableToolPoint.cpp
new file mode 100644
--- /dev/null
+++ b/Engine/Plugins/Runtime/ScriptableToolsFramework/Source/ScriptableToolsFramework/Private/Drawing/ScriptableToolPoint.cpp
@@ -0,0 +1,49 @@
+// Copyright Epic Games, Inc. All Rights Reserved.
+
+#include "Drawing/ScriptableToolPoint.h"
+
+
+void UScriptableToolPoint::SetPointID(int32 PointIDIn)
+{
+	PointID = PointIDIn;
+}
+
+int32 UScriptableToolPoint::GetPointID() const
+{
+	return PointID;
+}
+
+bool UScriptableToolPoint::IsDirty() const
+{
+	return bIsDirty;
+}
+
+FRenderablePoint UScriptableToolPoint::GeneratePointDescription()
+{
+	bIsDirty = false;
+	return PointDescription;
+}
+
+void UScriptableToolPoint::SetPointPosition(FVector PositionIn)
+{
+	PointDescription.Position = PositionIn;
+	bIsDirty = true;
+}
+
+void UScriptableToolPoint::SetPointColor(FColor PointColorIn)
+{
+	PointDescription.Color = PointColorIn;
+	bIsDirty = true;
+}
+
+void UScriptableToolPoint::SetPointSize(float SizeIn)
+{
+	PointDescription.Size = SizeIn;
+	bIsDirty = true;
+}
+
+void UScriptableToolPoint::SetPointDepthBias(float PointDepthBiasIn)
+{
+	PointDescription.DepthBias = PointDepthBiasIn;
+	bIsDirty = true;
+}
\ No newline at end of file
diff --git a/Engine/Plugins/Runtime/ScriptableToolsFramework/Source/ScriptableToolsFramework/Private/Drawing/ScriptableToolPointSet.cpp b/Engine/Plugins/Runtime/ScriptableToolsFramework/Source/ScriptableToolsFramework/Private/Drawing/ScriptableToolPointSet.cpp
new file mode 100644
--- /dev/null
+++ b/Engine/Plugins/Runtime/ScriptableToolsFramework/Source/ScriptableToolsFramework/Private/Drawing/ScriptableToolPointSet.cpp
@@ -0,0 +1,62 @@
+// Copyright Epic Games, Inc. All Rights Reserved.
+
+#include "Drawing/ScriptableToolPointSet.h"
+
+#include "Drawing/ScriptableToolPoint.h"
+#include "Drawing/PreviewGeometryActor.h"
+#include "Misc/Guid.h"
+
+void UScriptableToolPointSet::Initialize(TObjectPtr<UPreviewGeometry> PreviewGeometry)
+{
+	FString PointSetID = FGuid::NewGuid().ToString();
+	PointSet = PreviewGeometry->AddPointSet(PointSetID);
+}
+
+void UScriptableToolPointSet::OnTick()
+{
+	for (TObjectPtr<UScriptableToolPoint> PointComponent : PointComponents)
+	{
+		if (PointComponent->IsDirty())
+		{
+			int32 PointID = PointComponent->GetPointID();
+			FRenderablePoint PointDescription = PointComponent->GeneratePointDescription();
+
+			PointSet->SetPointPosition(PointID, PointDescription.Position);			
+			PointSet->SetPointColor(PointID, PointDescription.Color);
+			PointSet->SetPointSize(PointID, PointDescription.Size);
+		}
+	}
+}
+
+UScriptableToolPoint* UScriptableToolPointSet::AddPoint()
+{
+	PointComponents.Add(NewObject<UScriptableToolPoint>(this));
+	UScriptableToolPoint* NewPoint = PointComponents.Last();
+	NewPoint->SetPointID(PointSet->AddPoint(NewPoint->GeneratePointDescription()));
+	return NewPoint;
+}
+
+void UScriptableToolPointSet::RemovePoint(UScriptableToolPoint* Point)
+{
+	if (ensure(Point))
+	{
+		PointSet->RemovePoint(Point->GetPointID());
+		PointComponents.Remove(Point);
+	}
+}
+
+void UScriptableToolPointSet::RemoveAllPoints()
+{
+	PointSet->Clear();
+	PointComponents.Empty();
+}
+
+void UScriptableToolPointSet::SetAllPointsColor(FColor Color)
+{
+	PointSet->SetAllPointsColor(Color);
+}
+
+void UScriptableToolPointSet::SetAllPointsSize(float Size)
+{
+	PointSet->SetAllPointsSize(Size);
+}
\ No newline at end of file
diff --git a/Engine/Plugins/Runtime/ScriptableToolsFramework/Source/ScriptableToolsFramework/Private/Drawing/ScriptableToolTriangle.cpp b/Engine/Plugins/Runtime/ScriptableToolsFramework/Source/ScriptableToolsFramework/Private/Drawing/ScriptableToolTriangle.cpp
new file mode 100644
--- /dev/null
+++ b/Engine/Plugins/Runtime/ScriptableToolsFramework/Source/ScriptableToolsFramework/Private/Drawing/ScriptableToolTriangle.cpp
@@ -0,0 +1,143 @@
+// Copyright Epic Games, Inc. All Rights Reserved.
+
+#include "Drawing/ScriptableToolTriangle.h"
+
+
+void UScriptableToolTriangle::SetTriangleID(int32 TriangleIDIn)
+{
+	TriangleID = TriangleIDIn;
+}
+
+int32 UScriptableToolTriangle::GetTriangleID() const
+{
+	return TriangleID;
+}
+
+bool UScriptableToolTriangle::IsDirty() const
+{
+	return bIsDirty;
+}
+
+FRenderableTriangle UScriptableToolTriangle::GenerateTriangleDescription()
+{
+	bIsDirty = false;
+	return TriangleDescription;
+}
+
+void UScriptableToolTriangle::SetTriangleMaterial(UMaterialInterface* Material)
+{
+	TriangleDescription.Material = Material;
+	bIsDirty = true;
+}
+
+void UScriptableToolTriangle::SetTrianglePoints(FVector A, FVector B, FVector C)
+{
+	TriangleDescription.Vertex0.Position = A;
+	TriangleDescription.Vertex1.Position = B;
+	TriangleDescription.Vertex2.Position = C;
+	bIsDirty = true;
+}
+
+void UScriptableToolTriangle::SetTriangleUVs(FVector2D A, FVector2D B, FVector2D C)
+{
+	TriangleDescription.Vertex0.UV = A;
+	TriangleDescription.Vertex1.UV = B;
+	TriangleDescription.Vertex2.UV = C;
+	bIsDirty = true;
+}
+
+void UScriptableToolTriangle::SetTriangleNormals(FVector A, FVector B, FVector C)
+{
+	TriangleDescription.Vertex0.Normal = A;
+	TriangleDescription.Vertex1.Normal = B;
+	TriangleDescription.Vertex2.Normal = C;
+	bIsDirty = true;
+}
+
+void UScriptableToolTriangle::SetTriangleColors(FColor A, FColor B, FColor C)
+{
+	TriangleDescription.Vertex0.Color = A;
+	TriangleDescription.Vertex1.Color = B;
+	TriangleDescription.Vertex2.Color = C;
+	bIsDirty = true;
+}
+
+
+
+
+
+
+
+void UScriptableToolQuad::SetTriangleAID(int32 TriangleIDIn)
+{
+	TriangleAID = TriangleIDIn;
+}
+
+int32 UScriptableToolQuad::GetTriangleAID() const
+{
+	return TriangleAID;
+}
+
+void UScriptableToolQuad::SetTriangleBID(int32 TriangleIDIn)
+{
+	TriangleBID = TriangleIDIn;
+}
+
+int32 UScriptableToolQuad::GetTriangleBID() const
+{
+	return TriangleBID;
+}
+
+bool UScriptableToolQuad::IsDirty() const
+{
+	return bIsDirty;
+}
+
+TPair<FRenderableTriangle, FRenderableTriangle> UScriptableToolQuad::GenerateQuadDescription()
+{
+	bIsDirty = false;
+	return TPair<FRenderableTriangle, FRenderableTriangle>(TriangleADescription, TriangleBDescription);
+}
+
+void UScriptableToolQuad::SetQuadMaterial(UMaterialInterface* Material)
+{
+	TriangleADescription.Material = Material;
+	TriangleBDescription.Material = Material;
+	bIsDirty = true;
+}
+
+void UScriptableToolQuad::SetQuadPoints(FVector A, FVector B, FVector C, FVector D)
+{
+	TriangleADescription.Vertex0.Position = TriangleBDescription.Vertex0.Position = A;
+	TriangleADescription.Vertex1.Position = B;
+	TriangleADescription.Vertex2.Position = TriangleBDescription.Vertex1.Position = C;
+	TriangleBDescription.Vertex2.Position = D;
+	bIsDirty = true;
+}
+
+void UScriptableToolQuad::SetQuadUVs(FVector2D A, FVector2D B, FVector2D C, FVector2D D)
+{
+	TriangleADescription.Vertex0.UV = TriangleBDescription.Vertex0.UV = A;
+	TriangleADescription.Vertex1.UV = B;
+	TriangleADescription.Vertex2.UV = TriangleBDescription.Vertex1.UV = C;
+	TriangleBDescription.Vertex2.UV = D;
+	bIsDirty = true;
+}
+
+void UScriptableToolQuad::SetQuadNormals(FVector A, FVector B, FVector C, FVector D)
+{
+	TriangleADescription.Vertex0.Normal = TriangleBDescription.Vertex0.Normal = A;
+	TriangleADescription.Vertex1.Normal = B;
+	TriangleADescription.Vertex2.Normal = TriangleBDescription.Vertex1.Normal = C;
+	TriangleBDescription.Vertex2.Normal = D;
+	bIsDirty = true;
+}
+
+void UScriptableToolQuad::SetQuadColors(FColor A, FColor B, FColor C, FColor D)
+{
+	TriangleADescription.Vertex0.Color = TriangleBDescription.Vertex0.Color = A;
+	TriangleADescription.Vertex1.Color = B;
+	TriangleADescription.Vertex2.Color = TriangleBDescription.Vertex1.Color = C;
+	TriangleBDescription.Vertex2.Color = D;
+	bIsDirty = true;
+}
diff --git a/Engine/Plugins/Runtime/ScriptableToolsFramework/Source/ScriptableToolsFramework/Private/Drawing/ScriptableToolTriangleSet.cpp b/Engine/Plugins/Runtime/ScriptableToolsFramework/Source/ScriptableToolsFramework/Private/Drawing/ScriptableToolTriangleSet.cpp
new file mode 100644
--- /dev/null
+++ b/Engine/Plugins/Runtime/ScriptableToolsFramework/Source/ScriptableToolsFramework/Private/Drawing/ScriptableToolTriangleSet.cpp
@@ -0,0 +1,105 @@
+// Copyright Epic Games, Inc. All Rights Reserved.
+
+#include "Drawing/ScriptableToolTriangleSet.h"
+
+#include "Drawing/ScriptableToolTriangle.h"
+#include "Drawing/PreviewGeometryActor.h"
+#include "Misc/Guid.h"
+
+void UScriptableToolTriangleSet::Initialize(TObjectPtr<UPreviewGeometry> PreviewGeometry)
+{
+	FString TriangleSetID = FGuid::NewGuid().ToString();
+	TriangleSet = PreviewGeometry->AddTriangleSet(TriangleSetID);
+}
+
+void UScriptableToolTriangleSet::OnTick()
+{
+	for (TObjectPtr<UScriptableToolTriangle> TriangleComponent : TriangleComponents)
+	{
+		if (TriangleComponent->IsDirty())
+		{
+			int32 TriangleID = TriangleComponent->GetTriangleID();
+			FRenderableTriangle TriangleDescription = TriangleComponent->GenerateTriangleDescription();
+
+			TriangleSet->RemoveTriangle(TriangleID);
+			TriangleSet->InsertTriangle(TriangleID, TriangleDescription);
+		}
+	}
+
+	for (TObjectPtr<UScriptableToolQuad> QuadComponent : QuadComponents)
+	{
+		if (QuadComponent->IsDirty())
+		{
+			int32 TriangleAID = QuadComponent->GetTriangleAID();
+			int32 TriangleBID = QuadComponent->GetTriangleBID();
+
+			TPair<FRenderableTriangle, FRenderableTriangle> TriangleDescriptions = QuadComponent->GenerateQuadDescription();
+
+			TriangleSet->RemoveTriangle(TriangleAID);
+			TriangleSet->InsertTriangle(TriangleAID, TriangleDescriptions.Key);
+			TriangleSet->RemoveTriangle(TriangleBID);
+			TriangleSet->InsertTriangle(TriangleBID, TriangleDescriptions.Value);
+
+		}
+	}
+}
+
+UScriptableToolTriangle* UScriptableToolTriangleSet::AddTriangle()
+{
+	TriangleComponents.Add(NewObject<UScriptableToolTriangle>(this));
+	UScriptableToolTriangle* NewTriangle = TriangleComponents.Last();
+	NewTriangle->SetTriangleID(TriangleSet->AddTriangle(NewTriangle->GenerateTriangleDescription()));
+	return NewTriangle;
+}
+
+UScriptableToolQuad* UScriptableToolTriangleSet::AddQuad()
+{
+	QuadComponents.Add(NewObject<UScriptableToolQuad>(this));
+	UScriptableToolQuad* NewQuad = QuadComponents.Last();
+
+	TPair<FRenderableTriangle, FRenderableTriangle> TriDescriptions = NewQuad->GenerateQuadDescription();
+
+	int32 TriAID = TriangleSet->AddTriangle(TriDescriptions.Key);
+	int32 TriBID = TriangleSet->AddTriangle(TriDescriptions.Value);
+
+	NewQuad->SetTriangleAID(TriAID);
+	NewQuad->SetTriangleBID(TriBID);
+
+	return NewQuad;
+}
+
+void UScriptableToolTriangleSet::RemoveTriangle(UScriptableToolTriangle* Triangle)
+{
+	if (ensure(Triangle))
+	{
+		TriangleSet->RemoveTriangle(Triangle->GetTriangleID());
+		TriangleComponents.Remove(Triangle);
+	}
+}
+
+void UScriptableToolTriangleSet::RemoveQuad(UScriptableToolQuad* Quad)
+{
+	if (ensure(Quad))
+	{
+		TriangleSet->RemoveTriangle(Quad->GetTriangleAID());
+		TriangleSet->RemoveTriangle(Quad->GetTriangleBID());
+		QuadComponents.Remove(Quad);
+	}
+}
+
+
+void UScriptableToolTriangleSet::RemoveAllFaces()
+{
+	TriangleSet->Clear();
+	TriangleComponents.Empty();
+}
+
+void UScriptableToolTriangleSet::SetAllTrianglesColor(FColor Color)
+{
+	TriangleSet->SetAllTrianglesColor(Color);
+}
+
+void UScriptableToolTriangleSet::SetAllTrianglesMaterial(UMaterialInterface* Material)
+{
+	TriangleSet->SetAllTrianglesMaterial(Material);
+}
\ No newline at end of file
diff --git a/Engine/Plugins/Runtime/ScriptableToolsFramework/Source/ScriptableToolsFramework/Private/ScriptableInteractiveTool.cpp b/Engine/Plugins/Runtime/ScriptableToolsFramework/Source/ScriptableToolsFramework/Private/ScriptableInteractiveTool.cpp
--- a/Engine/Plugins/Runtime/ScriptableToolsFramework/Source/ScriptableToolsFramework/Private/ScriptableInteractiveTool.cpp
+++ b/Engine/Plugins/Runtime/ScriptableToolsFramework/Source/ScriptableToolsFramework/Private/ScriptableInteractiveTool.cpp
@@ -8,6 +8,12 @@
 #include "BaseGizmos/TransformGizmoUtil.h"
 #include "BaseGizmos/CombinedTransformGizmo.h"
 
+#include "Drawing/PreviewGeometryActor.h"
+#include "Drawing/ScriptableToolLineSet.h"
+#include "Drawing/ScriptableToolPointSet.h"
+#include "Drawing/ScriptableToolTriangleSet.h"
+
+
 #include "Engine/Font.h"
 #include "ToolDataVisualizer.h"
 #include "CanvasTypes.h"
@@ -20,7 +26,6 @@
 #define LOCTEXT_NAMESPACE "UScriptableInteractiveTool"
 
 
-
 UScriptableInteractiveTool* UScriptableInteractiveToolPropertySet::GetOwningTool(EToolsFrameworkOutcomePins& Outcome)
 {
 	if (ParentTool.IsValid())
@@ -187,12 +192,43 @@ void UScriptableInteractiveTool::Setup()
 	DrawHUDHelper = NewObject<UScriptableTool_HUDAPI>();
 
 	OnScriptSetup();
+
+
+	ToolDrawableGeometry = NewObject<UPreviewGeometry>();
+	ToolDrawableGeometry->CreateInWorld(GetWorld(), FTransform::Identity);
+
+	DefaultLineSet = NewObject<UScriptableToolLineSet>();
+	DefaultLineSet->Initialize(ToolDrawableGeometry);
+
+	DefaultPointSet = NewObject<UScriptableToolPointSet>();
+	DefaultPointSet->Initialize(ToolDrawableGeometry);
+
+	DefaultTriangleSet = NewObject<UScriptableToolTriangleSet>();
+	DefaultTriangleSet->Initialize(ToolDrawableGeometry);
 }
 
 void UScriptableInteractiveTool::OnTick(float DeltaTime)
 {
 	UInteractiveTool::OnTick(DeltaTime);
 	OnScriptTick(DeltaTime);
+
+	DefaultLineSet->OnTick();
+	for (UScriptableToolLineSet* LineSet : LineSets)
+	{
+		LineSet->OnTick();
+	}
+
+	DefaultPointSet->OnTick();
+	for (UScriptableToolPointSet* PointSet : PointSets)
+	{
+		PointSet->OnTick();
+	}
+
+	DefaultTriangleSet->OnTick();
+	for (UScriptableToolTriangleSet* TriangleSet : TriangleSets)
+	{
+		TriangleSet->OnTick();
+	}
 }
 
 bool UScriptableInteractiveTool::HasAccept() const
@@ -233,6 +269,9 @@ void UScriptableInteractiveTool::Shutdown(EToolShutdownType ShutdownType)
 	RenderHelper = nullptr;
 	DrawHUDHelper = nullptr;
 
+	ToolDrawableGeometry->Disconnect();
+	ToolDrawableGeometry = nullptr;
+
 	UInteractiveTool::Shutdown(ShutdownType);
 }
 
@@ -1141,6 +1180,48 @@ TArray<UToolTarget*> UScriptableInteractiveTool::GetToolTargets() const
 
 
 
+UScriptableToolLineSet* UScriptableInteractiveTool::GetDefaultLineSet() const
+{
+	return DefaultLineSet;
+}
+
+
+UScriptableToolLineSet* UScriptableInteractiveTool::AddLineSet()
+{
+	TObjectPtr<UScriptableToolLineSet> LineSet = NewObject<UScriptableToolLineSet>();
+	LineSet->Initialize(ToolDrawableGeometry);
+	LineSets.Add( LineSet );
+	return LineSet;
+}
+
+UScriptableToolPointSet* UScriptableInteractiveTool::GetDefaultPointSet() const
+{
+	return DefaultPointSet;
+}
+
+
+UScriptableToolPointSet* UScriptableInteractiveTool::AddPointSet()
+{
+	TObjectPtr<UScriptableToolPointSet> PointSet = NewObject<UScriptableToolPointSet>();
+	PointSet->Initialize(ToolDrawableGeometry);
+	PointSets.Add(PointSet);
+	return PointSet;
+}
+
+
+UScriptableToolTriangleSet* UScriptableInteractiveTool::GetDefaultTriangleSet() const
+{
+	return DefaultTriangleSet;
+}
+
+
+UScriptableToolTriangleSet* UScriptableInteractiveTool::AddTriangleSet()
+{
+	TObjectPtr<UScriptableToolTriangleSet> TriangleSet = NewObject<UScriptableToolTriangleSet>();
+	TriangleSet->Initialize(ToolDrawableGeometry);
+	TriangleSets.Add(TriangleSet);
+	return TriangleSet;
+}
 
 
 
diff --git a/Engine/Plugins/Runtime/ScriptableToolsFramework/Source/ScriptableToolsFramework/Public/Drawing/ScriptableToolLine.h b/Engine/Plugins/Runtime/ScriptableToolsFramework/Source/ScriptableToolsFramework/Public/Drawing/ScriptableToolLine.h
new file mode 100644
--- /dev/null
+++ b/Engine/Plugins/Runtime/ScriptableToolsFramework/Source/ScriptableToolsFramework/Public/Drawing/ScriptableToolLine.h
@@ -0,0 +1,81 @@
+// Copyright Epic Games, Inc. All Rights Reserved.
+
+#pragma once
+
+
+#include "ScriptableInteractiveTool.h"
+#include "Drawing/LineSetComponent.h"
+
+
+#include "ScriptableToolLine.generated.h"
+
+
+UCLASS(BlueprintType)
+class SCRIPTABLETOOLSFRAMEWORK_API UScriptableToolLine : public UObject
+{
+	GENERATED_BODY()
+
+public:
+
+	void SetLineID(int32 LineIDIn);
+	int32 GetLineID() const;
+
+	bool IsDirty() const;
+	FRenderableLine GenerateLineDescription();
+
+	/**
+	 * Set the starting position of the line
+	 * @param Start The position in space of the start of the line
+	 */
+	UFUNCTION(BlueprintCallable, Category = "ScriptableTool|Drawing|Lines")
+	void SetLineStart(FVector Start);
+
+	/**
+	 * Set the ending position of the line
+	 * @param End The position in space of the end of the line
+	 */
+	UFUNCTION(BlueprintCallable, Category = "ScriptableTool|Drawing|Lines")
+	void SetLineEnd(FVector End);
+
+	/**
+	 * Set the starting and ending positions of the line
+	 * @param Start The position in space of the start of the line
+	 * @param End The position in space of the end of the line
+	 */
+	UFUNCTION(BlueprintCallable, Category = "ScriptableTool|Drawing|Lines")
+	void SetLineEndPoints(FVector Start, FVector End);
+
+	/**
+	 * Set the line's color
+	 * @param Color The color the line should be rendered as
+	 */
+	UFUNCTION(BlueprintCallable, Category = "ScriptableTool|Drawing|Lines")
+	void SetLineColor(FColor Color);
+
+	/**
+	 * Set the line's thickness
+	 * @param Thickness The visual thickness in pixels the line should be rendered as
+	 */
+	UFUNCTION(BlueprintCallable, Category = "ScriptableTool|Drawing|Lines")
+	void SetLineThickness(float Thickness);
+
+	/**
+	 * Set the line's depth bias. The depth bias controls small micro adjustments in effective
+	 * displacement towards or away from the camera, to prevent the line from causing clipping
+	 * or z-fighting with other objects in the scene it's overlaid directly on top of.
+	 * @param DepthBias The depth bias adjustment for the line
+	 */
+	UFUNCTION(BlueprintCallable, Category = "ScriptableTool|Drawing|Lines")
+	void SetLineDepthBias(float DepthBias);
+
+private:
+
+	FRenderableLine LineDescription;
+
+	UPROPERTY()
+	bool bIsDirty = false;
+
+	UPROPERTY()
+	int32 LineID;
+
+};
\ No newline at end of file
diff --git a/Engine/Plugins/Runtime/ScriptableToolsFramework/Source/ScriptableToolsFramework/Public/Drawing/ScriptableToolLineSet.h b/Engine/Plugins/Runtime/ScriptableToolsFramework/Source/ScriptableToolsFramework/Public/Drawing/ScriptableToolLineSet.h
new file mode 100644
--- /dev/null
+++ b/Engine/Plugins/Runtime/ScriptableToolsFramework/Source/ScriptableToolsFramework/Public/Drawing/ScriptableToolLineSet.h
@@ -0,0 +1,70 @@
+// Copyright Epic Games, Inc. All Rights Reserved.
+
+#pragma once
+
+#include "ScriptableInteractiveTool.h"
+#include "Drawing/LineSetComponent.h"
+
+#include "ScriptableToolLineSet.generated.h"
+
+
+class UScriptableToolLine;
+class UPreviewGeometry;
+
+UCLASS(BlueprintType)
+class SCRIPTABLETOOLSFRAMEWORK_API UScriptableToolLineSet : public UObject
+{
+	GENERATED_BODY()
+
+public:
+
+	void Initialize(TObjectPtr<UPreviewGeometry> PreviewGeometry);
+
+	void OnTick();
+
+
+	/**
+	 * Create and return a new line object. Users should save a reference to this object for future updates or removal from the set.
+	 * @return The new line object added to the set.
+	 */
+	UFUNCTION(BlueprintCallable, Category = "ScriptableTool|Drawing|LineSet")
+	UPARAM(DisplayName = "LineComponent") UScriptableToolLine* AddLine();
+
+	/**
+	 * Remove a specific line object from the set, removing it fromt the scene.
+	 * @param Line A reference to a line to be removed from the set.
+	 */
+	UFUNCTION(BlueprintCallable, Category = "ScriptableTool|Drawing|LineSet")
+	void RemoveLine(UScriptableToolLine* Line);
+
+
+	/**
+	 * Remove all current lines in the set.
+	 */
+	UFUNCTION(BlueprintCallable, Category = "ScriptableTool|Drawing|LineSet")
+	void RemoveAllLines();
+
+	/**
+	 * Set the color of all lines in the set simultaneously.
+	 */
+	UFUNCTION(BlueprintCallable, Category = "ScriptableTool|Drawing|LineSet")
+	void SetAllLinesColor(FColor Color);
+
+	/**
+	 * Set the thickness of all lines in the set simultaneously.
+	 */
+	UFUNCTION(BlueprintCallable, Category = "ScriptableTool|Drawing|LineSet")
+	void SetAllLinesThickness(float Thickness);
+
+protected:
+
+	UPROPERTY(Transient, DuplicateTransient, NonTransactional, SkipSerialization)
+	TObjectPtr<UPreviewGeometry> ToolDrawableGeometry = nullptr;
+
+	UPROPERTY(Transient, DuplicateTransient, NonTransactional, SkipSerialization)
+	TObjectPtr<ULineSetComponent> LineSet = nullptr;
+
+	TArray<TObjectPtr<UScriptableToolLine>> LineComponents;
+
+};
+
diff --git a/Engine/Plugins/Runtime/ScriptableToolsFramework/Source/ScriptableToolsFramework/Public/Drawing/ScriptableToolPoint.h b/Engine/Plugins/Runtime/ScriptableToolsFramework/Source/ScriptableToolsFramework/Public/Drawing/ScriptableToolPoint.h
new file mode 100644
--- /dev/null
+++ b/Engine/Plugins/Runtime/ScriptableToolsFramework/Source/ScriptableToolsFramework/Public/Drawing/ScriptableToolPoint.h
@@ -0,0 +1,66 @@
+// Copyright Epic Games, Inc. All Rights Reserved.
+
+#pragma once
+
+
+#include "ScriptableInteractiveTool.h"
+#include "Drawing/PointSetComponent.h"
+
+
+#include "ScriptableToolPoint.generated.h"
+
+
+UCLASS(BlueprintType)
+class SCRIPTABLETOOLSFRAMEWORK_API UScriptableToolPoint : public UObject
+{
+	GENERATED_BODY()
+
+public:
+
+	void SetPointID(int32 PointIDIn);
+	int32 GetPointID() const;
+
+	bool IsDirty() const;
+	FRenderablePoint GeneratePointDescription();
+
+	/**
+	 * Set the position of the point
+	 * @param Position The position in space of the point
+	 */
+	UFUNCTION(BlueprintCallable, Category = "ScriptableTool|Drawing|Points")
+	void SetPointPosition(FVector Position);
+
+	/**
+	 * Set the point's color
+	 * @param Color The color the point should be rendered as
+	 */
+	UFUNCTION(BlueprintCallable, Category = "ScriptableTool|Drawing|Points")
+	void SetPointColor(FColor Color);
+
+	/**
+	 * Set the point's size
+	 * @param Size The visual size in pixels the point should be rendered as
+	 */
+	UFUNCTION(BlueprintCallable, Category = "ScriptableTool|Drawing|Points")
+	void SetPointSize(float Thickness);
+
+	/**
+	 * Set the point's depth bias. The depth bias controls small micro adjustments in effective
+	 * displacement towards or away from the camera, to prevent the point from causing clipping
+	 * or z-fighting with other objects in the scene it's overlaid directly on top of.
+	 * @param DepthBias The depth bias adjustment for the point
+	 */
+	UFUNCTION(BlueprintCallable, Category = "ScriptableTool|Drawing|Points")
+	void SetPointDepthBias(float DepthBias);
+
+private:
+
+	FRenderablePoint PointDescription;
+
+	UPROPERTY()
+	bool bIsDirty = false;
+
+	UPROPERTY()
+	int32 PointID;
+
+};
\ No newline at end of file
diff --git a/Engine/Plugins/Runtime/ScriptableToolsFramework/Source/ScriptableToolsFramework/Public/Drawing/ScriptableToolPointSet.h b/Engine/Plugins/Runtime/ScriptableToolsFramework/Source/ScriptableToolsFramework/Public/Drawing/ScriptableToolPointSet.h
new file mode 100644
--- /dev/null
+++ b/Engine/Plugins/Runtime/ScriptableToolsFramework/Source/ScriptableToolsFramework/Public/Drawing/ScriptableToolPointSet.h
@@ -0,0 +1,69 @@
+// Copyright Epic Games, Inc. All Rights Reserved.
+
+#pragma once
+
+#include "ScriptableInteractiveTool.h"
+#include "Drawing/PointSetComponent.h"
+
+#include "ScriptableToolPointSet.generated.h"
+
+
+class UScriptableToolPoint;
+class UPreviewGeometry;
+
+UCLASS(BlueprintType)
+class SCRIPTABLETOOLSFRAMEWORK_API UScriptableToolPointSet : public UObject
+{
+	GENERATED_BODY()
+
+public:
+
+	void Initialize(TObjectPtr<UPreviewGeometry> PreviewGeometry);
+
+	void OnTick();
+
+	/**
+	 * Create and return a new point object. Users should save a reference to this object for future updates or removal from the set.
+	 * @return The new point object added to the set.
+	 */
+	UFUNCTION(BlueprintCallable, Category = "ScriptableTool|Drawing|PointSet")
+	UPARAM(DisplayName = "PointComponent") UScriptableToolPoint* AddPoint();
+
+
+	/**
+	 * Remove a specific point object from the set, removing it from the scene.
+	 * @param Point A reference to a point to be removed from the set.
+	 */
+	UFUNCTION(BlueprintCallable, Category = "ScriptableTool|Drawing|PointSet")
+	void RemovePoint(UScriptableToolPoint* Point);
+
+	/**
+	 * Remove all current points in the set.
+	 */
+	UFUNCTION(BlueprintCallable, Category = "ScriptableTool|Drawing|PointSet")
+	void RemoveAllPoints();
+
+	/**
+	 * Set the color of all points in the set simultaneously.
+	 */
+	UFUNCTION(BlueprintCallable, Category = "ScriptableTool|Drawing|PointSet")
+	void SetAllPointsColor(FColor Color);
+
+	/**
+	 * Set the size of all points in the set simultaneously.
+	 */
+	UFUNCTION(BlueprintCallable, Category = "ScriptableTool|Drawing|PointSet")
+	void SetAllPointsSize(float Size);
+
+protected:
+
+	UPROPERTY(Transient, DuplicateTransient, NonTransactional, SkipSerialization)
+	TObjectPtr<UPreviewGeometry> ToolDrawableGeometry = nullptr;
+
+	UPROPERTY(Transient, DuplicateTransient, NonTransactional, SkipSerialization)
+	TObjectPtr<UPointSetComponent> PointSet = nullptr;
+
+	TArray<TObjectPtr<UScriptableToolPoint>> PointComponents;
+
+};
+
diff --git a/Engine/Plugins/Runtime/ScriptableToolsFramework/Source/ScriptableToolsFramework/Public/Drawing/ScriptableToolTriangle.h b/Engine/Plugins/Runtime/ScriptableToolsFramework/Source/ScriptableToolsFramework/Public/Drawing/ScriptableToolTriangle.h
new file mode 100644
--- /dev/null
+++ b/Engine/Plugins/Runtime/ScriptableToolsFramework/Source/ScriptableToolsFramework/Public/Drawing/ScriptableToolTriangle.h
@@ -0,0 +1,164 @@
+// Copyright Epic Games, Inc. All Rights Reserved.
+
+#pragma once
+
+
+#include "ScriptableInteractiveTool.h"
+#include "Drawing/TriangleSetComponent.h"
+
+
+#include "ScriptableToolTriangle.generated.h"
+
+
+UCLASS(BlueprintType)
+class SCRIPTABLETOOLSFRAMEWORK_API UScriptableToolTriangle : public UObject
+{
+	GENERATED_BODY()
+
+public:
+
+	//void Initialize();
+
+	void SetTriangleID(int32 TriangleIDIn);
+	int32 GetTriangleID() const;
+
+	bool IsDirty() const;
+	FRenderableTriangle GenerateTriangleDescription();
+
+	/*
+	* Set the material of the triangle
+	* 
+	* @param Material The new material that should be assigned to the triangle
+	*/
+	UFUNCTION(BlueprintCallable, Category = "ScriptableTool|Drawing|Triangles")
+	void SetTriangleMaterial(UMaterialInterface* Material);
+
+	/*
+	* Set the points of the triangle
+	*
+	* @param A The position of the first corner of the triangle
+	* @param B The position of the second corner of the triangle
+	* @param C The position of the third corner of the triangle
+	*/
+	UFUNCTION(BlueprintCallable, Category = "ScriptableTool|Drawing|Triangles")
+	void SetTrianglePoints(FVector A, FVector B, FVector C);
+
+	/*
+	* Set the UV coordinates of the triangle
+	*
+	* @param A The UV coordinate of the first corner of the triangle
+	* @param B The UV coordinate of the second corner of the triangle
+	* @param C The UV coordinate of the third corner of the triangle
+	*/
+	UFUNCTION(BlueprintCallable, Category = "ScriptableTool|Drawing|Triangles")
+	void SetTriangleUVs(FVector2D A, FVector2D B, FVector2D C);
+
+	/*
+	* Set the normal values of the triangle
+	*
+	* @param A The normal of the first corner of the triangle
+	* @param B The normal of the second corner of the triangle
+	* @param C The normal of the third corner of the triangle
+	*/
+	UFUNCTION(BlueprintCallable, Category = "ScriptableTool|Drawing|Triangles")
+	void SetTriangleNormals(FVector A, FVector B, FVector C);
+
+
+	/*
+	* Set the vertex colors of the triangle
+	*
+	* @param A The vertex color of the first corner of the triangle
+	* @param B The vertex color of the second corner of the triangle
+	* @param C The vertex color of the third corner of the triangle
+	*/
+	UFUNCTION(BlueprintCallable, Category = "ScriptableTool|Drawing|Triangles")
+	void SetTriangleColors(FColor A, FColor B, FColor C);
+
+private:
+
+	FRenderableTriangle TriangleDescription;
+
+	UPROPERTY()
+	bool bIsDirty = false;
+
+	UPROPERTY()
+	int32 TriangleID;
+
+};
+
+UCLASS(BlueprintType)
+class SCRIPTABLETOOLSFRAMEWORK_API UScriptableToolQuad : public UObject
+{
+	GENERATED_BODY()
+
+public:
+
+	void SetTriangleAID(int32 TriangleIDIn);
+	int32 GetTriangleAID() const;
+	void SetTriangleBID(int32 TriangleIDIn);
+	int32 GetTriangleBID() const;
+
+	bool IsDirty() const;
+	TPair<FRenderableTriangle, FRenderableTriangle> GenerateQuadDescription();
+
+	/**
+	 * Set the material of the quad
+	 * @param Material The new material that should be assigned to the quad
+	 */
+	UFUNCTION(BlueprintCallable, Category = "ScriptableTool|Drawing|Quads")
+	void SetQuadMaterial(UMaterialInterface* Material);
+
+	/**
+	 * Set the points of the quad
+	 * @param A The position of the first corner of the quad
+	 * @param B The position of the second corner of the quad
+	 * @param C The position of the third corner of the quad
+	 * @param D The position of the fourth corner of the quad
+	 */
+	UFUNCTION(BlueprintCallable, Category = "ScriptableTool|Drawing|Quads")
+	void SetQuadPoints(FVector A, FVector B, FVector C, FVector D);
+
+	/**
+	 * Set the UV coordinates of the quad
+	 * @param A The UV coordinate of the first corner of the quad
+	 * @param B The UV coordinate of the second corner of the quad
+	 * @param C The UV coordinate of the third corner of the quad
+	 * @param D The UV coordinate of the fourth corner of the quad
+	 */
+	UFUNCTION(BlueprintCallable, Category = "ScriptableTool|Drawing|Quads")
+	void SetQuadUVs(FVector2D A, FVector2D B, FVector2D C, FVector2D D);
+
+	/**
+	 * Set the normal values of the quad
+	 * @param A The normal of the first corner of the quad
+	 * @param B The normal of the second corner of the quad
+	 * @param C The normal of the third corner of the quad
+	 * @param D The normal of the fourth corner of the quad
+	 */
+	UFUNCTION(BlueprintCallable, Category = "ScriptableTool|Drawing|Quads")
+	void SetQuadNormals(FVector A, FVector B, FVector C, FVector D);
+
+	/**
+	 * Set the vertex colors of the quad
+	 * @param A The vertex color of the first corner of the quad
+	 * @param B The vertex color of the second corner of the quad
+	 * @param C The vertex color of the third corner of the quad
+	 * @param D The vertex color of the fourth corner of the quad
+	 */
+	UFUNCTION(BlueprintCallable, Category = "ScriptableTool|Drawing|Quads")
+	void SetQuadColors(FColor A, FColor B, FColor C, FColor D);
+private:
+
+	FRenderableTriangle TriangleADescription;
+	FRenderableTriangle TriangleBDescription;
+
+	UPROPERTY()
+	bool bIsDirty = false;
+
+	UPROPERTY()
+	int32 TriangleAID;
+
+	UPROPERTY()
+	int32 TriangleBID;
+
+};
\ No newline at end of file
diff --git a/Engine/Plugins/Runtime/ScriptableToolsFramework/Source/ScriptableToolsFramework/Public/Drawing/ScriptableToolTriangleSet.h b/Engine/Plugins/Runtime/ScriptableToolsFramework/Source/ScriptableToolsFramework/Public/Drawing/ScriptableToolTriangleSet.h
new file mode 100644
--- /dev/null
+++ b/Engine/Plugins/Runtime/ScriptableToolsFramework/Source/ScriptableToolsFramework/Public/Drawing/ScriptableToolTriangleSet.h
@@ -0,0 +1,85 @@
+// Copyright Epic Games, Inc. All Rights Reserved.
+
+#pragma once
+
+#include "ScriptableInteractiveTool.h"
+#include "Drawing/TriangleSetComponent.h"
+
+#include "ScriptableToolTriangleSet.generated.h"
+
+
+class UScriptableToolTriangle;
+class UScriptableToolQuad;
+class UPreviewGeometry;
+
+UCLASS(BlueprintType)
+class SCRIPTABLETOOLSFRAMEWORK_API UScriptableToolTriangleSet : public UObject
+{
+	GENERATED_BODY()
+
+public:
+
+	void Initialize(TObjectPtr<UPreviewGeometry> PreviewGeometry);
+
+	void OnTick();
+
+	/**
+	 * Create and return a new triangle object. Users should save a reference to this object for future updates or removal from the set.
+	 * @return The new triangle object added to the set
+	 */
+	UFUNCTION(BlueprintCallable, Category = "ScriptableTool|Drawing|TriangleSet")
+	UPARAM(DisplayName = "TriangleComponent") UScriptableToolTriangle* AddTriangle();
+
+	/**
+	 * Create and return a new quad object. Quad objects are two paired triangles. Users should save a reference to this object for future updates or removal from the set.
+	 * @return The new quad object added to the set
+	 */
+	UFUNCTION(BlueprintCallable, Category = "ScriptableTool|Drawing|TriangleSet")
+	UPARAM(DisplayName = "QuadComponent") UScriptableToolQuad* AddQuad();
+
+	/**
+	 * Remove a specific triangle object from the set, removing it from the scene.
+	 * @param Triangle A reference to a triangle to be removed from the set.
+	 */
+	UFUNCTION(BlueprintCallable, Category = "ScriptableTool|Drawing|TriangleSet")
+	void RemoveTriangle(UScriptableToolTriangle* Triangle);
+
+	/**
+	 * Remove a specific quad object from the set, removing it from the scene.
+	 * @param Quad A reference to a quad to be removed from the set.
+	 */
+	UFUNCTION(BlueprintCallable, Category = "ScriptableTool|Drawing|TriangleSet")
+	void RemoveQuad(UScriptableToolQuad* Quad);
+
+	/**
+	 * Remove all current triangles and quads in the set. 
+	 */
+	UFUNCTION(BlueprintCallable, Category = "ScriptableTool|Drawing|TriangleSet")
+	void RemoveAllFaces();
+
+	/**
+	 * Set the color of all triangles and quads in the set simultaneously .
+	 */
+	UFUNCTION(BlueprintCallable, Category = "ScriptableTool|Drawing|TriangleSet")
+	void SetAllTrianglesColor(FColor Color);
+
+	/**
+	 * Set the material of all triangles and quads in the set simultaneously .
+	 */
+	UFUNCTION(BlueprintCallable, Category = "ScriptableTool|Drawing|TriangleSet")
+	void SetAllTrianglesMaterial(UMaterialInterface* Material);
+
+protected:
+
+	UPROPERTY(Transient, DuplicateTransient, NonTransactional, SkipSerialization)
+	TObjectPtr<UPreviewGeometry> ToolDrawableGeometry = nullptr;
+
+	UPROPERTY(Transient, DuplicateTransient, NonTransactional, SkipSerialization)
+	TObjectPtr<UTriangleSetComponent> TriangleSet = nullptr;
+
+	TArray<TObjectPtr<UScriptableToolTriangle>> TriangleComponents;
+
+	TArray < TObjectPtr < UScriptableToolQuad>> QuadComponents;
+
+};
+
diff --git a/Engine/Plugins/Runtime/ScriptableToolsFramework/Source/ScriptableToolsFramework/Public/ScriptableInteractiveTool.h b/Engine/Plugins/Runtime/ScriptableToolsFramework/Source/ScriptableToolsFramework/Public/ScriptableInteractiveTool.h
--- a/Engine/Plugins/Runtime/ScriptableToolsFramework/Source/ScriptableToolsFramework/Public/ScriptableInteractiveTool.h
+++ b/Engine/Plugins/Runtime/ScriptableToolsFramework/Source/ScriptableToolsFramework/Public/ScriptableInteractiveTool.h
@@ -17,6 +17,12 @@ class FToolDataVisualizer;
 class FCanvas;
 class FSceneView;
 
+class UPreviewGeometry;
+class UScriptableToolLineSet;
+class UScriptableToolPointSet;
+class UScriptableToolTriangleSet;
+
+
 UENUM(BlueprintType)
 enum class EToolsFrameworkOutcomePins : uint8
 {
@@ -802,11 +808,79 @@ protected:
 	UPROPERTY(Transient, DuplicateTransient, NonTransactional, SkipSerialization)
 	TArray<TObjectPtr<UToolTarget>> Targets;
 
-};
 
+public:
+
+	// Drawing API
 
+	/**
+	 * Retrieve the default line set object for the tool, used for drawing persistent line objects in the scene.
+	 * @return A reference to the tool's default line set object
+	 */
+	UFUNCTION(BlueprintCallable, Category = "ScriptableTool|Drawing")
+	UScriptableToolLineSet* GetDefaultLineSet() const;
 
+	/**
+	 * Create and return a new, independent line set, used for drawing persistent line objects in the scene.
+	 * Users must save a reference to the created object for future access.
+	 * @return A reference to a new line set object
+	 */
+	UFUNCTION(BlueprintCallable, Category = "ScriptableTool|Drawing")
+	UScriptableToolLineSet* AddLineSet();
 
+	/**
+	 * Retrieve the default point set object for the tool, used for drawing persistent point objects in the scene.
+	 * @return A reference to the tool's default point set object
+	 */
+	UFUNCTION(BlueprintCallable, Category = "ScriptableTool|Drawing")
+	UScriptableToolPointSet* GetDefaultPointSet() const;
+
+	/**
+	 * Create and return a new, independent point set, used for drawing persistent point objects in the scene.
+	 * Users must save a reference to the created object for future access.
+	 * @return A reference to a new point set object
+	 */
+	UFUNCTION(BlueprintCallable, Category = "ScriptableTool|Drawing")
+	UScriptableToolPointSet* AddPointSet();
+
+	/**
+	 * Retrieve the default triangle set object for the tool, used for drawing persistent triangle and quad objects in the scene.
+	 * @return A reference to the tool's default triangle set object
+	 */
+	UFUNCTION(BlueprintCallable, Category = "ScriptableTool|Drawing")
+	UScriptableToolTriangleSet* GetDefaultTriangleSet() const;
+
+	/**
+	 * Create and return a new, independent triangle set, used for drawing persistent triangle and quad objects in the scene.
+	 * Users must save a reference to the created object for future access.
+	 * @return A reference to a new triangle set object
+	 */
+	UFUNCTION(BlueprintCallable, Category = "ScriptableTool|Drawing")
+	UScriptableToolTriangleSet* AddTriangleSet();
+
+protected:
+	UPROPERTY(Transient, DuplicateTransient, NonTransactional, SkipSerialization)
+	TObjectPtr<UScriptableToolLineSet> DefaultLineSet;
+
+	UPROPERTY(Transient, DuplicateTransient, NonTransactional, SkipSerialization)
+	TArray<TObjectPtr<UScriptableToolLineSet> > LineSets;
+
+	UPROPERTY(Transient, DuplicateTransient, NonTransactional, SkipSerialization)
+	TObjectPtr<UScriptableToolPointSet> DefaultPointSet;
+
+	UPROPERTY(Transient, DuplicateTransient, NonTransactional, SkipSerialization)
+	TArray<TObjectPtr<UScriptableToolPointSet> > PointSets;
+
+	UPROPERTY(Transient, DuplicateTransient, NonTransactional, SkipSerialization)
+	TObjectPtr<UScriptableToolTriangleSet> DefaultTriangleSet;
+
+	UPROPERTY(Transient, DuplicateTransient, NonTransactional, SkipSerialization)
+	TArray<TObjectPtr<UScriptableToolTriangleSet> > TriangleSets;
+
+	UPROPERTY(Transient, DuplicateTransient, NonTransactional, SkipSerialization)
+	TObjectPtr<UPreviewGeometry> ToolDrawableGeometry = nullptr;
+
+};
 
 
 UCLASS(meta = (ScriptName = "ScriptableTools_Util"))
diff --git a/Engine/Plugins/Runtime/ScriptableToolsFramework/Source/ScriptableToolsFramework/ScriptableToolsFramework.Build.cs b/Engine/Plugins/Runtime/ScriptableToolsFramework/Source/ScriptableToolsFramework/ScriptableToolsFramework.Build.cs
--- a/Engine/Plugins/Runtime/ScriptableToolsFramework/Source/ScriptableToolsFramework/ScriptableToolsFramework.Build.cs
+++ b/Engine/Plugins/Runtime/ScriptableToolsFramework/Source/ScriptableToolsFramework/ScriptableToolsFramework.Build.cs
@@ -31,7 +31,8 @@ public class ScriptableToolsFramework : ModuleRules
 				"RenderCore",
 				"GeometryCore",
 				"InputCore",
-				"InteractiveToolsFramework"
+				"InteractiveToolsFramework",
+				"ModelingComponents"
 			}
 			);
 
