ScriptableTools: Added DoubleClickBehavior support

diff --git a/Engine/Plugins/Runtime/ScriptableToolsFramework/Source/ScriptableToolsFramework/Private/BaseTools/ScriptableModularBehaviorTool.cpp b/Engine/Plugins/Runtime/ScriptableToolsFramework/Source/ScriptableToolsFramework/Private/BaseTools/ScriptableModularBehaviorTool.cpp
--- a/Engine/Plugins/Runtime/ScriptableToolsFramework/Source/ScriptableToolsFramework/Private/BaseTools/ScriptableModularBehaviorTool.cpp
+++ b/Engine/Plugins/Runtime/ScriptableToolsFramework/Source/ScriptableToolsFramework/Private/BaseTools/ScriptableModularBehaviorTool.cpp
@@ -4,6 +4,7 @@
 
 #include "Behaviors/ScriptableToolSingleClickBehavior.h"
 #include "Behaviors/ScriptableToolClickDragBehavior.h"
+#include "Behaviors/ScriptableToolDoubleClickBehavior.h"
 #include "Behaviors/ScriptableToolMouseHoverBehavior.h"
 #include "Behaviors/ScriptableToolMouseWheelBehavior.h"
 #include "Behaviors/ScriptableToolMultiClickSequenceBehavior.h"
@@ -26,6 +27,22 @@ void UScriptableModularBehaviorTool::AddSingleClickBehavior(
 }
 
 
+void UScriptableModularBehaviorTool::AddDoubleClickBehavior(
+	const FTestIfHitByClickDelegate IfHitByClick,
+	const FOnHitByClickDelegate OnHitByClick,
+	const FMouseBehaviorModiferCheckDelegate ModifierCheckFunction,
+	int CapturePriority,
+	EScriptableToolMouseButton MouseButton,
+	bool bHitTestOnRelease)
+{
+	TObjectPtr<UScriptableToolDoubleClickBehavior> BehaviorContainer = NewObject<UScriptableToolDoubleClickBehavior>();
+	BehaviorContainer->Init(this, ModifierCheckFunction, IfHitByClick, OnHitByClick, MouseButton, bHitTestOnRelease);
+	BehaviorContainer->SetDefaultPriority(FInputCapturePriority(CapturePriority));
+
+	DoubleClickBehaviors.Add(BehaviorContainer);
+}
+
+
 void UScriptableModularBehaviorTool::AddClickDragBehavior(
 	const FTestCanBeginClickDragSequenceDelegate TestCanBeginClickDragSequenceFuncIn,
 	const FOnClickPressDelegate OnClickPressFuncIn,
diff --git a/Engine/Plugins/Runtime/ScriptableToolsFramework/Source/ScriptableToolsFramework/Private/Behaviors/ScriptableToolDoubleClickBehavior.cpp b/Engine/Plugins/Runtime/ScriptableToolsFramework/Source/ScriptableToolsFramework/Private/Behaviors/ScriptableToolDoubleClickBehavior.cpp
new file mode 100644
--- /dev/null
+++ b/Engine/Plugins/Runtime/ScriptableToolsFramework/Source/ScriptableToolsFramework/Private/Behaviors/ScriptableToolDoubleClickBehavior.cpp
@@ -0,0 +1,9 @@
+// Copyright Epic Games, Inc. All Rights Reserved.
+
+#include "Behaviors/ScriptableToolDoubleClickBehavior.h"
+#include "BaseBehaviors/DoubleClickBehavior.h"
+
+USingleClickInputBehavior* UScriptableToolDoubleClickBehavior::CreateNewBehavior() const
+{
+	return NewObject<UDoubleClickInputBehavior>();
+}
diff --git a/Engine/Plugins/Runtime/ScriptableToolsFramework/Source/ScriptableToolsFramework/Private/Behaviors/ScriptableToolDoubleClickBehavior.h b/Engine/Plugins/Runtime/ScriptableToolsFramework/Source/ScriptableToolsFramework/Private/Behaviors/ScriptableToolDoubleClickBehavior.h
new file mode 100644
--- /dev/null
+++ b/Engine/Plugins/Runtime/ScriptableToolsFramework/Source/ScriptableToolsFramework/Private/Behaviors/ScriptableToolDoubleClickBehavior.h
@@ -0,0 +1,18 @@
+// Copyright Epic Games, Inc. All Rights Reserved.
+
+#pragma once
+
+#include "ScriptableToolSingleClickBehavior.h"
+
+#include "ScriptableToolDoubleClickBehavior.generated.h"
+
+UCLASS()
+class UScriptableToolDoubleClickBehavior : public UScriptableToolSingleClickBehavior
+{
+	GENERATED_BODY()
+
+public:
+	UScriptableToolDoubleClickBehavior() {};
+
+	virtual USingleClickInputBehavior* CreateNewBehavior() const override;
+};
diff --git a/Engine/Plugins/Runtime/ScriptableToolsFramework/Source/ScriptableToolsFramework/Private/Behaviors/ScriptableToolSingleClickBehavior.cpp b/Engine/Plugins/Runtime/ScriptableToolsFramework/Source/ScriptableToolsFramework/Private/Behaviors/ScriptableToolSingleClickBehavior.cpp
--- a/Engine/Plugins/Runtime/ScriptableToolsFramework/Source/ScriptableToolsFramework/Private/Behaviors/ScriptableToolSingleClickBehavior.cpp
+++ b/Engine/Plugins/Runtime/ScriptableToolsFramework/Source/ScriptableToolsFramework/Private/Behaviors/ScriptableToolSingleClickBehavior.cpp
@@ -15,7 +15,7 @@ void UScriptableToolSingleClickBehavior::Init(
 	)
 {
 	BehaviorHost = BehaviorHostIn;
-	Behavior = NewObject<USingleClickInputBehavior>();
+	Behavior = CreateNewBehavior();
 	ModifierCheckFunc = ModifierCheckFuncIn;
 	TestIfHitByClickFunc = TestIfHitByClickFuncIn;
 	OnHitByClickFunc = OnHitByClickFuncIn;
@@ -61,6 +61,11 @@ void UScriptableToolSingleClickBehavior::Init(
 	}
 }
 
+USingleClickInputBehavior* UScriptableToolSingleClickBehavior::CreateNewBehavior() const
+{
+	return NewObject<USingleClickInputBehavior>();
+}
+
 UInputBehavior* UScriptableToolSingleClickBehavior::GetWrappedBehavior()
 {
 	return Behavior;
diff --git a/Engine/Plugins/Runtime/ScriptableToolsFramework/Source/ScriptableToolsFramework/Private/Behaviors/ScriptableToolSingleClickBehavior.h b/Engine/Plugins/Runtime/ScriptableToolsFramework/Source/ScriptableToolsFramework/Private/Behaviors/ScriptableToolSingleClickBehavior.h
--- a/Engine/Plugins/Runtime/ScriptableToolsFramework/Source/ScriptableToolsFramework/Private/Behaviors/ScriptableToolSingleClickBehavior.h
+++ b/Engine/Plugins/Runtime/ScriptableToolsFramework/Source/ScriptableToolsFramework/Private/Behaviors/ScriptableToolSingleClickBehavior.h
@@ -26,6 +26,7 @@ public:
 		EScriptableToolMouseButton MouseButtonIn,
 		bool bHitTestOnReleaseIn);
 
+	virtual USingleClickInputBehavior* CreateNewBehavior() const;
 	virtual UInputBehavior* GetWrappedBehavior() override;
 
 	virtual FInputRayHit IsHitByClick(const FInputDeviceRay& ClickPos) override;
diff --git a/Engine/Plugins/Runtime/ScriptableToolsFramework/Source/ScriptableToolsFramework/Public/BaseTools/ScriptableModularBehaviorTool.h b/Engine/Plugins/Runtime/ScriptableToolsFramework/Source/ScriptableToolsFramework/Public/BaseTools/ScriptableModularBehaviorTool.h
--- a/Engine/Plugins/Runtime/ScriptableToolsFramework/Source/ScriptableToolsFramework/Public/BaseTools/ScriptableModularBehaviorTool.h
+++ b/Engine/Plugins/Runtime/ScriptableToolsFramework/Source/ScriptableToolsFramework/Public/BaseTools/ScriptableModularBehaviorTool.h
@@ -9,6 +9,7 @@
 #include "ScriptableModularBehaviorTool.generated.h"
 
 class UScriptableToolSingleClickBehavior;
+class UScriptableToolDoubleClickBehavior;
 class UScriptableToolClickDragBehavior;
 class UScriptableToolSingleClickOrDragBehavior;
 class UScriptableToolMouseWheelBehavior;
@@ -54,6 +55,32 @@ public:
 		bool bHitTestOnRelease = true
 	);
 
+	/**
+	*   Implements a standard "button-click"-style input behavior
+	*   The state machine works as follows:
+	*    1) on input-device-button-press, hit-test the target. If hit, begin capture
+	*    2) on input-device-button-release, hit-test the target. If hit, call OnClicked(). If not hit, ignore click.
+	*    
+	*   The second hit-test is required to allow the click to be "cancelled" by moving away
+	*   from the target. This is standard GUI behavior. You can disable this second hit test
+	*   using the HitTestOnRelease property. This is strongly discouraged.
+	* 
+	*	@param IfHitByClick Test if hit by a click
+	*	@param OnHitByClick Notify that double click occurred
+	*	@param CaptureCheck Only enable capture if returns true
+	*	@param CapturePriority The priority is used to resolve situations where multiple behaviors want the same capture
+	*	@param MouseButton Determines which mouse button the behavior captures on
+	*/
+	UFUNCTION(BlueprintCallable, Category = "ScriptableTool|Input", meta = (AdvancedDisplay = "CaptureCheck, CapturePriority, MouseButton, bHitTestOnRelease"))
+	void AddDoubleClickBehavior(
+		const FTestIfHitByClickDelegate IfHitByClick,
+		const FOnHitByClickDelegate OnHitByClick,
+		const FMouseBehaviorModiferCheckDelegate CaptureCheck,
+		int CapturePriority = 100,
+		EScriptableToolMouseButton MouseButton = EScriptableToolMouseButton::LeftButton,
+		bool bHitTestOnRelease = true
+	);
+
 	/**
 	*   Implements a standard "button-click-drag"-style input behavior.
 	*	 
@@ -204,6 +231,9 @@ private:
 	UPROPERTY(Transient, DuplicateTransient, NonTransactional, SkipSerialization)
 	TArray< TObjectPtr<UScriptableToolSingleClickBehavior> > SingleClickBehaviors;
 
+	UPROPERTY(Transient, DuplicateTransient, NonTransactional, SkipSerialization)
+	TArray< TObjectPtr<UScriptableToolDoubleClickBehavior> > DoubleClickBehaviors;
+
 	UPROPERTY(Transient, DuplicateTransient, NonTransactional, SkipSerialization)
 	TArray< TObjectPtr<UScriptableToolClickDragBehavior> > ClickDragBehaviors;
 
