ScriptableTools: Fix crash bug that occurs when launching the ScriptableTools mode with no scriptable tools found by the editor.

diff --git a/Engine/Plugins/Runtime/ScriptableToolsFramework/Source/EditorScriptableToolsFramework/Private/ScriptableToolSet.cpp b/Engine/Plugins/Runtime/ScriptableToolsFramework/Source/EditorScriptableToolsFramework/Private/ScriptableToolSet.cpp
--- a/Engine/Plugins/Runtime/ScriptableToolsFramework/Source/EditorScriptableToolsFramework/Private/ScriptableToolSet.cpp
+++ b/Engine/Plugins/Runtime/ScriptableToolsFramework/Source/EditorScriptableToolsFramework/Private/ScriptableToolSet.cpp
@@ -72,7 +72,14 @@ void UScriptableToolSet::ReinitializeScriptableTools(FPreToolsLoadedDelegate Pre
 
 	AsyncLoadHandle = UAssetManager::GetStreamableManager().RequestAsyncLoad(ObjectPathsToLoad, [this, PostDelegate, ObjectPathsToLoad, TagsToFilterCopy]() { PostToolLoad(PostDelegate, ObjectPathsToLoad, TagsToFilterCopy);  });
 
-	AsyncLoadHandle->BindUpdateDelegate(UpdateDelegate);
+	if (AsyncLoadHandle)
+	{
+		AsyncLoadHandle->BindUpdateDelegate(UpdateDelegate);
+	}
+	else
+	{
+		PostToolLoad(PostDelegate, ObjectPathsToLoad, TagsToFilterCopy);
+	}
 }
 
 
