ScriptableTools: Fix issue where garbage collection could cause scriptable tools to crash when drawing geometry.

diff --git a/Engine/Plugins/Runtime/ScriptableToolsFramework/Source/ScriptableToolsFramework/Public/Drawing/ScriptableToolLineSet.h b/Engine/Plugins/Runtime/ScriptableToolsFramework/Source/ScriptableToolsFramework/Public/Drawing/ScriptableToolLineSet.h
--- a/Engine/Plugins/Runtime/ScriptableToolsFramework/Source/ScriptableToolsFramework/Public/Drawing/ScriptableToolLineSet.h
+++ b/Engine/Plugins/Runtime/ScriptableToolsFramework/Source/ScriptableToolsFramework/Public/Drawing/ScriptableToolLineSet.h
@@ -64,6 +64,7 @@ protected:
 	UPROPERTY(Transient, DuplicateTransient, NonTransactional, SkipSerialization)
 	TObjectPtr<ULineSetComponent> LineSet = nullptr;
 
+	UPROPERTY(Transient, DuplicateTransient, NonTransactional, SkipSerialization)
 	TArray<TObjectPtr<UScriptableToolLine>> LineComponents;
 
 };
diff --git a/Engine/Plugins/Runtime/ScriptableToolsFramework/Source/ScriptableToolsFramework/Public/Drawing/ScriptableToolPointSet.h b/Engine/Plugins/Runtime/ScriptableToolsFramework/Source/ScriptableToolsFramework/Public/Drawing/ScriptableToolPointSet.h
--- a/Engine/Plugins/Runtime/ScriptableToolsFramework/Source/ScriptableToolsFramework/Public/Drawing/ScriptableToolPointSet.h
+++ b/Engine/Plugins/Runtime/ScriptableToolsFramework/Source/ScriptableToolsFramework/Public/Drawing/ScriptableToolPointSet.h
@@ -63,6 +63,7 @@ protected:
 	UPROPERTY(Transient, DuplicateTransient, NonTransactional, SkipSerialization)
 	TObjectPtr<UPointSetComponent> PointSet = nullptr;
 
+	UPROPERTY(Transient, DuplicateTransient, NonTransactional, SkipSerialization)
 	TArray<TObjectPtr<UScriptableToolPoint>> PointComponents;
 
 };
diff --git a/Engine/Plugins/Runtime/ScriptableToolsFramework/Source/ScriptableToolsFramework/Public/Drawing/ScriptableToolTriangleSet.h b/Engine/Plugins/Runtime/ScriptableToolsFramework/Source/ScriptableToolsFramework/Public/Drawing/ScriptableToolTriangleSet.h
--- a/Engine/Plugins/Runtime/ScriptableToolsFramework/Source/ScriptableToolsFramework/Public/Drawing/ScriptableToolTriangleSet.h
+++ b/Engine/Plugins/Runtime/ScriptableToolsFramework/Source/ScriptableToolsFramework/Public/Drawing/ScriptableToolTriangleSet.h
@@ -77,8 +77,10 @@ protected:
 	UPROPERTY(Transient, DuplicateTransient, NonTransactional, SkipSerialization)
 	TObjectPtr<UTriangleSetComponent> TriangleSet = nullptr;
 
+	UPROPERTY(Transient, DuplicateTransient, NonTransactional, SkipSerialization)
 	TArray<TObjectPtr<UScriptableToolTriangle>> TriangleComponents;
 
+	UPROPERTY(Transient, DuplicateTransient, NonTransactional, SkipSerialization)
 	TArray < TObjectPtr < UScriptableToolQuad>> QuadComponents;
 
 };
