ScriptableTools: Implement support for the new HoverCapture check function in the Scriptable Tools wrapped behaviors for Mouse Hover and MultiClick behaviors.

diff --git a/Engine/Plugins/Runtime/ScriptableToolsFramework/Source/ScriptableToolsFramework/Private/BaseTools/ScriptableModularBehaviorTool.cpp b/Engine/Plugins/Runtime/ScriptableToolsFramework/Source/ScriptableToolsFramework/Private/BaseTools/ScriptableModularBehaviorTool.cpp
--- a/Engine/Plugins/Runtime/ScriptableToolsFramework/Source/ScriptableToolsFramework/Private/BaseTools/ScriptableModularBehaviorTool.cpp
+++ b/Engine/Plugins/Runtime/ScriptableToolsFramework/Source/ScriptableToolsFramework/Private/BaseTools/ScriptableModularBehaviorTool.cpp
@@ -117,12 +117,13 @@ void UScriptableModularBehaviorTool::AddMultiClickSequenceBehavior(
 	FOnTerminateClickSequenceDelegate OnTerminateClickSequenceFuncIn,
 	FRequestAbortClickSequenceDelegate RequestAbortClickSequenceFuncIn,
 	FMouseBehaviorModiferCheckDelegate ModifierCheckFuncIn,
+	FMouseBehaviorModiferCheckDelegate HoverModifierCheckFuncIn,
 	int CapturePriority,
 	EScriptableToolMouseButton MouseButtonIn
 )
 {
 	TObjectPtr<UScriptableToolClickSequenceBehavior> BehaviorContainer = NewObject<UScriptableToolClickSequenceBehavior>();
-	BehaviorContainer->Init(this, ModifierCheckFuncIn, OnBeginSequencePreviewFuncIn, CanBeginClickSequenceFuncIn, OnBeginClickSequenceFuncIn, OnNextSequencePreviewFuncIn,
+	BehaviorContainer->Init(this, ModifierCheckFuncIn, HoverModifierCheckFuncIn, OnBeginSequencePreviewFuncIn, CanBeginClickSequenceFuncIn, OnBeginClickSequenceFuncIn, OnNextSequencePreviewFuncIn,
 		                    OnNextSequenceClickFuncIn, OnTerminateClickSequenceFuncIn, RequestAbortClickSequenceFuncIn, MouseButtonIn);
 	BehaviorContainer->SetDefaultPriority(FInputCapturePriority(CapturePriority));
 
@@ -137,11 +138,12 @@ void UScriptableModularBehaviorTool::AddMouseHoverBehavior(
 	FOnBeginHoverDelegate OnBeginHoverFuncIn,
 	FOnUpdateHoverDelegate OnUpdateHoverFuncIn,
 	FOnEndHoverDelegate OnEndHoverFuncIn,
+	FMouseBehaviorModiferCheckDelegate HoverModifierCheckFuncIn,
 	int CapturePriority
 )
 {
 	TObjectPtr<UScriptableToolMouseHoverBehavior> BehaviorContainer = NewObject<UScriptableToolMouseHoverBehavior>();
-	BehaviorContainer->Init(this, BeginHoverSequenceHitTestFuncIn, OnBeginHoverFuncIn, OnUpdateHoverFuncIn, OnEndHoverFuncIn);
+	BehaviorContainer->Init(this, HoverModifierCheckFuncIn, BeginHoverSequenceHitTestFuncIn, OnBeginHoverFuncIn, OnUpdateHoverFuncIn, OnEndHoverFuncIn);
 	BehaviorContainer->SetDefaultPriority(FInputCapturePriority(CapturePriority));
 
 	MouseHoverBehaviors.Add(BehaviorContainer);
diff --git a/Engine/Plugins/Runtime/ScriptableToolsFramework/Source/ScriptableToolsFramework/Private/Behaviors/ScriptableToolMouseHoverBehavior.cpp b/Engine/Plugins/Runtime/ScriptableToolsFramework/Source/ScriptableToolsFramework/Private/Behaviors/ScriptableToolMouseHoverBehavior.cpp
--- a/Engine/Plugins/Runtime/ScriptableToolsFramework/Source/ScriptableToolsFramework/Private/Behaviors/ScriptableToolMouseHoverBehavior.cpp
+++ b/Engine/Plugins/Runtime/ScriptableToolsFramework/Source/ScriptableToolsFramework/Private/Behaviors/ScriptableToolMouseHoverBehavior.cpp
@@ -7,6 +7,7 @@
 
 
 void UScriptableToolMouseHoverBehavior::Init(TObjectPtr<UScriptableModularBehaviorTool> BehaviorHostIn,
+	FMouseBehaviorModiferCheckDelegate HoverModifierCheckFuncIn,
 	FBeginHoverSequenceHitTestDelegate BeginHoverSequenceHitTestFuncIn,
 	FOnBeginHoverDelegate OnBeginHoverFuncIn,
 	FOnUpdateHoverDelegate OnUpdateHoverFuncIn,
@@ -14,6 +15,7 @@ void UScriptableToolMouseHoverBehavior::Init(TObjectPtr<UScriptableModularBehavi
 {
 	BehaviorHost = BehaviorHostIn;
 	Behavior = NewObject<UMouseHoverBehavior>();
+	HoverModifierCheckFunc = HoverModifierCheckFuncIn;
 	BeginHoverSequenceHitTestFunc = BeginHoverSequenceHitTestFuncIn;
 	OnBeginHoverFunc = OnBeginHoverFuncIn;
 	OnUpdateHoverFunc = OnUpdateHoverFuncIn;
@@ -21,6 +23,15 @@ void UScriptableToolMouseHoverBehavior::Init(TObjectPtr<UScriptableModularBehavi
 
 	Behavior->Initialize(this);
 
+	Behavior->HoverModifierCheckFunc = [this](const FInputDeviceState& InputDeviceState)
+	{
+		bool bResult = true;
+		if (HoverModifierCheckFunc.IsBound())
+		{
+			bResult = HoverModifierCheckFunc.Execute(InputDeviceState);
+		}
+		return bResult;
+	};
 
 	BehaviorHost->AddInputBehavior(Behavior);
 
diff --git a/Engine/Plugins/Runtime/ScriptableToolsFramework/Source/ScriptableToolsFramework/Private/Behaviors/ScriptableToolMouseHoverBehavior.h b/Engine/Plugins/Runtime/ScriptableToolsFramework/Source/ScriptableToolsFramework/Private/Behaviors/ScriptableToolMouseHoverBehavior.h
--- a/Engine/Plugins/Runtime/ScriptableToolsFramework/Source/ScriptableToolsFramework/Private/Behaviors/ScriptableToolMouseHoverBehavior.h
+++ b/Engine/Plugins/Runtime/ScriptableToolsFramework/Source/ScriptableToolsFramework/Private/Behaviors/ScriptableToolMouseHoverBehavior.h
@@ -21,6 +21,7 @@ public:
 	UScriptableToolMouseHoverBehavior() {};
 
 	void Init(TObjectPtr<UScriptableModularBehaviorTool> BehaviorHostIn,
+		FMouseBehaviorModiferCheckDelegate HoverModifierCheckFuncIn,
 		FBeginHoverSequenceHitTestDelegate BeginHoverSequenceHitTestFuncIn,
 		FOnBeginHoverDelegate OnBeginHoverFuncIn,
 		FOnUpdateHoverDelegate OnUpdateHoverFuncIn,
@@ -42,6 +43,9 @@ private:
 	UPROPERTY()
 	TObjectPtr<UMouseHoverBehavior> Behavior;
 
+	UPROPERTY()
+	FMouseBehaviorModiferCheckDelegate HoverModifierCheckFunc;
+
 	UPROPERTY()
 	FBeginHoverSequenceHitTestDelegate BeginHoverSequenceHitTestFunc;
 
diff --git a/Engine/Plugins/Runtime/ScriptableToolsFramework/Source/ScriptableToolsFramework/Private/Behaviors/ScriptableToolMultiClickSequenceBehavior.cpp b/Engine/Plugins/Runtime/ScriptableToolsFramework/Source/ScriptableToolsFramework/Private/Behaviors/ScriptableToolMultiClickSequenceBehavior.cpp
--- a/Engine/Plugins/Runtime/ScriptableToolsFramework/Source/ScriptableToolsFramework/Private/Behaviors/ScriptableToolMultiClickSequenceBehavior.cpp
+++ b/Engine/Plugins/Runtime/ScriptableToolsFramework/Source/ScriptableToolsFramework/Private/Behaviors/ScriptableToolMultiClickSequenceBehavior.cpp
@@ -7,6 +7,7 @@
 
 void UScriptableToolClickSequenceBehavior::Init(TObjectPtr<UScriptableModularBehaviorTool> BehaviorHostIn,
 	FMouseBehaviorModiferCheckDelegate ModifierCheckFuncIn,
+	FMouseBehaviorModiferCheckDelegate HoverModifierCheckFuncIn,
 	FOnBeginSequencePreviewDelegate OnBeginSequencePreviewFuncIn,
 	FCanBeginClickSequenceDelegate CanBeginClickSequenceFuncIn,
 	FOnBeginClickSequenceDelegate OnBeginClickSequenceFuncIn,
@@ -19,6 +20,7 @@ void UScriptableToolClickSequenceBehavior::Init(TObjectPtr<UScriptableModularBeh
 	BehaviorHost = BehaviorHostIn;
 	Behavior = NewObject<UMultiClickSequenceInputBehavior>();
 	ModifierCheckFunc = ModifierCheckFuncIn;
+	HoverModifierCheckFunc = HoverModifierCheckFuncIn;
 	OnBeginSequencePreviewFunc = OnBeginSequencePreviewFuncIn;
 	CanBeginClickSequenceFunc = CanBeginClickSequenceFuncIn;
 	OnBeginClickSequenceFunc = OnBeginClickSequenceFuncIn;
@@ -38,6 +40,14 @@ void UScriptableToolClickSequenceBehavior::Init(TObjectPtr<UScriptableModularBeh
 		return bResult;
 	};
 
+	Behavior->HoverModifierCheckFunc = [this](const FInputDeviceState& InputDeviceState) {
+		bool bResult = true;
+		if (HoverModifierCheckFunc.IsBound())
+		{
+			bResult = HoverModifierCheckFunc.Execute(InputDeviceState);
+		}
+		return bResult;
+	};
 
 	BehaviorHost->AddInputBehavior(Behavior);
 
diff --git a/Engine/Plugins/Runtime/ScriptableToolsFramework/Source/ScriptableToolsFramework/Private/Behaviors/ScriptableToolMultiClickSequenceBehavior.h b/Engine/Plugins/Runtime/ScriptableToolsFramework/Source/ScriptableToolsFramework/Private/Behaviors/ScriptableToolMultiClickSequenceBehavior.h
--- a/Engine/Plugins/Runtime/ScriptableToolsFramework/Source/ScriptableToolsFramework/Private/Behaviors/ScriptableToolMultiClickSequenceBehavior.h
+++ b/Engine/Plugins/Runtime/ScriptableToolsFramework/Source/ScriptableToolsFramework/Private/Behaviors/ScriptableToolMultiClickSequenceBehavior.h
@@ -22,6 +22,7 @@ public:
 
 	void Init(TObjectPtr<UScriptableModularBehaviorTool> BehaviorHostIn,
 		FMouseBehaviorModiferCheckDelegate ModifierCheckFuncIn,
+		FMouseBehaviorModiferCheckDelegate HoverModifierCheckFuncIn,
 		FOnBeginSequencePreviewDelegate OnBeginSequencePreviewFuncIn,
 		FCanBeginClickSequenceDelegate CanBeginClickSequenceFuncIn,
 		FOnBeginClickSequenceDelegate OnBeginClickSequenceFuncIn,
@@ -53,6 +54,9 @@ private:
 	UPROPERTY()
 	FMouseBehaviorModiferCheckDelegate ModifierCheckFunc;
 
+	UPROPERTY()
+	FMouseBehaviorModiferCheckDelegate HoverModifierCheckFunc;
+
 	UPROPERTY()
 	FOnBeginSequencePreviewDelegate OnBeginSequencePreviewFunc;
 
diff --git a/Engine/Plugins/Runtime/ScriptableToolsFramework/Source/ScriptableToolsFramework/Public/BaseTools/ScriptableModularBehaviorTool.h b/Engine/Plugins/Runtime/ScriptableToolsFramework/Source/ScriptableToolsFramework/Public/BaseTools/ScriptableModularBehaviorTool.h
--- a/Engine/Plugins/Runtime/ScriptableToolsFramework/Source/ScriptableToolsFramework/Public/BaseTools/ScriptableModularBehaviorTool.h
+++ b/Engine/Plugins/Runtime/ScriptableToolsFramework/Source/ScriptableToolsFramework/Public/BaseTools/ScriptableModularBehaviorTool.h
@@ -193,10 +193,11 @@ public:
 	*	@param OnTerminateClickSequence Notify Target that click sequence has been explicitly terminated (eg by escape key, cancel tool, etc).
 	*	@param RequestAbortClickSequence Target overrides this and returns true if it wants to abort click sequence, checked every update
 	*	@param CaptureCheck Only enable capture if returns true
+	*   @param HoverCaptureCheck  Only enable hover capture if returns true
 	*	@param CapturePriority The priority is used to resolve situations where multiple behaviors want the same capture
 	*	@param MouseButton Determines which mouse button the behavior captures on
 	*/
-	UFUNCTION(BlueprintCallable, Category = "ScriptableTool|Input", meta = (AdvancedDisplay = "CaptureCheck, CapturePriority, MouseButton"))
+	UFUNCTION(BlueprintCallable, Category = "ScriptableTool|Input", meta = (AdvancedDisplay = "CaptureCheck, HoverCaptureCheck, CapturePriority, MouseButton"))
 	void AddMultiClickSequenceBehavior(		
 		FOnBeginSequencePreviewDelegate OnBeginSequencePreview,
 		FCanBeginClickSequenceDelegate CanBeginClickSequence,
@@ -206,6 +207,7 @@ public:
 		FOnTerminateClickSequenceDelegate OnTerminateClickSequence,
 		FRequestAbortClickSequenceDelegate RequestAbortClickSequence,
 		FMouseBehaviorModiferCheckDelegate CaptureCheck,
+		const FMouseBehaviorModiferCheckDelegate HoverCaptureCheck,
 		int CapturePriority = 100,
 		EScriptableToolMouseButton MouseButton = EScriptableToolMouseButton::LeftButton
 	);
@@ -215,14 +217,16 @@ public:
 	*	@param OnBeginHover Initialize hover sequence at given position
 	*	@param OnUpdateHover Update active hover sequence with new input position
 	*	@param OnEndHover Terminate active hover sequence
+	*   @param HoverCaptureCheck  Only enable hover capture if returns true
 	*	@param CapturePriority The priority is used to resolve situations where multiple behaviors want the same capture
 	*/
-	UFUNCTION(BlueprintCallable, Category = "ScriptableTool|Input", meta = (AdvancedDisplay = "CapturePriority"))
+	UFUNCTION(BlueprintCallable, Category = "ScriptableTool|Input", meta = (AdvancedDisplay = "CaptureCheck, CapturePriority"))
 	void AddMouseHoverBehavior(
 		FBeginHoverSequenceHitTestDelegate BeginHoverSequenceHitTest,
 		FOnBeginHoverDelegate OnBeginHover,
 		FOnUpdateHoverDelegate OnUpdateHover,
 		FOnEndHoverDelegate OnEndHover,
+		const FMouseBehaviorModiferCheckDelegate HoverCaptureCheck,
 		int CapturePriority = 100
 	);
 
