ScriptableTools: Improve deletion handling within the ScriptableTool Mode to prevent crashes when scriptable tool blueprints are deleted.

* Added checks to prevent deletion from happening on any tools or tool builders currently loaded by the mode.

* Added safety checks around groups to prevent weird behavior if null groups are introduced

* Added logging features for scriptable tools to describe failed deletion attempts

diff --git a/Engine/Plugins/Editor/ScriptableToolsEditorMode/Source/ScriptableToolsEditorMode/Private/ScriptableToolsEditorModeToolkit.cpp b/Engine/Plugins/Editor/ScriptableToolsEditorMode/Source/ScriptableToolsEditorMode/Private/ScriptableToolsEditorModeToolkit.cpp
--- a/Engine/Plugins/Editor/ScriptableToolsEditorMode/Source/ScriptableToolsEditorMode/Private/ScriptableToolsEditorModeToolkit.cpp
+++ b/Engine/Plugins/Editor/ScriptableToolsEditorMode/Source/ScriptableToolsEditorMode/Private/ScriptableToolsEditorModeToolkit.cpp
@@ -211,7 +211,9 @@ public:
 							if (BlueprintObject->GeneratedClass->IsChildOf(UScriptableToolGroupTag::StaticClass()))
 							{
 								TSubclassOf<UScriptableToolGroupTag> TagSubclass{ BlueprintObject->GeneratedClass };
-								ModeSettings->ToolRegistrationFilters.Groups.Add(TagSubclass);								
+								FScriptableToolGroupSet::FGroupSet  Groups = ModeSettings->ToolRegistrationFilters.GetGroups();
+								Groups.Add(TagSubclass);
+								ModeSettings->ToolRegistrationFilters.SetGroups(Groups);
 							}
 						}
 						ModeSettings->PostEditChange();
@@ -251,7 +253,7 @@ public:
 
 												TArray<FText> GroupNames;
 
-												for (UClass* GroupClass : ModeSettings->ToolRegistrationFilters.Groups)
+												for (UClass* GroupClass : ModeSettings->ToolRegistrationFilters.GetGroups())
 												{
 													if (GroupClass)
 													{
@@ -264,19 +266,20 @@ public:
 													}
 												}
 
-												if (ModeSettings->ToolRegistrationFilters.Groups.Num() == 1)
+												if (GroupNames.Num() == 0)
+												{
+													return LOCTEXT("ScriptableToolsZeroGroupLabel", "Showing tools from no groups");
+												}
+												else if (GroupNames.Num() == 1)
 												{
-													ensure(GroupNames.IsValidIndex(0));
 													return FText::Format(LOCTEXT("ScriptableToolsOneGroupLabel", "Showing tools from {0}"), GroupNames[0]);
 												}
-												else if (ModeSettings->ToolRegistrationFilters.Groups.Num() == 2)
+												else if (GroupNames.Num() == 2)
 												{
-													ensure(GroupNames.IsValidIndex(1));
 													return FText::Format(LOCTEXT("ScriptableToolsTwoGroupLabel", "Showing tools from {0} and {1}"), GroupNames[0], GroupNames[1]);
 												}
 												else
 												{
-													ensure(GroupNames.IsValidIndex(1));
 													return FText::Format(LOCTEXT("ScriptableToolsManyGroupLabel", "Showing tools from {0}, {1} and more..."), GroupNames[0], GroupNames[1]);
 												}
 											}
diff --git a/Engine/Plugins/Editor/ScriptableToolsEditorMode/Source/ScriptableToolsEditorMode/Private/Tags/EditableScriptableToolGroupSet.cpp b/Engine/Plugins/Editor/ScriptableToolsEditorMode/Source/ScriptableToolsEditorMode/Private/Tags/EditableScriptableToolGroupSet.cpp
--- a/Engine/Plugins/Editor/ScriptableToolsEditorMode/Source/ScriptableToolsEditorMode/Private/Tags/EditableScriptableToolGroupSet.cpp
+++ b/Engine/Plugins/Editor/ScriptableToolsEditorMode/Source/ScriptableToolsEditorMode/Private/Tags/EditableScriptableToolGroupSet.cpp
@@ -13,12 +13,12 @@ UEditableScriptableToolGroupSet::UEditableScriptableToolGroupSet()
 
 void UEditableScriptableToolGroupSet::SetGroups(const FScriptableToolGroupSet::FGroupSet& InGroups)
 {
-	GroupSet.Groups = InGroups;
+	GroupSet.SetGroups(InGroups);
 }
 
 FScriptableToolGroupSet::FGroupSet& UEditableScriptableToolGroupSet::GetGroups()
 {
-	return GroupSet.Groups;
+	return GroupSet.GetGroups();
 }
 
 FString UEditableScriptableToolGroupSet::GetGroupSetExportText()
diff --git a/Engine/Plugins/Editor/ScriptableToolsEditorMode/Source/ScriptableToolsEditorMode/Private/UI/SScriptableToolGroupSetCombo.cpp b/Engine/Plugins/Editor/ScriptableToolsEditorMode/Source/ScriptableToolsEditorMode/Private/UI/SScriptableToolGroupSetCombo.cpp
--- a/Engine/Plugins/Editor/ScriptableToolsEditorMode/Source/ScriptableToolsEditorMode/Private/UI/SScriptableToolGroupSetCombo.cpp
+++ b/Engine/Plugins/Editor/ScriptableToolsEditorMode/Source/ScriptableToolsEditorMode/Private/UI/SScriptableToolGroupSetCombo.cpp
@@ -168,7 +168,7 @@ FReply SScriptableToolGroupSetCombo::OnClearTagClicked(UClass* InGroup)
 	}
 	else if (StructPtr)
 	{
-		StructPtr->Groups = HelperGroupSet->GetGroups();
+		StructPtr->SetGroups(HelperGroupSet->GetGroups());
 	}
 
 	OnChanged.ExecuteIfBound();
@@ -217,7 +217,7 @@ void SScriptableToolGroupSetCombo::RefreshListView()
 	if(bValidData)
 	{
 		FScriptableToolGroupSet& GroupSet = *static_cast<FScriptableToolGroupSet*>(StructPointer);
-		FScriptableToolGroupSet::FGroupSet& Groups = GroupSet.Groups;
+		FScriptableToolGroupSet::FGroupSet Groups = GroupSet.GetGroups();
 
 		HelperGroupSet->SetGroups(Groups);
 
diff --git a/Engine/Plugins/Editor/ScriptableToolsEditorMode/Source/ScriptableToolsEditorMode/Private/UI/SScriptableToolGroupSetPicker.cpp b/Engine/Plugins/Editor/ScriptableToolsEditorMode/Source/ScriptableToolsEditorMode/Private/UI/SScriptableToolGroupSetPicker.cpp
--- a/Engine/Plugins/Editor/ScriptableToolsEditorMode/Source/ScriptableToolsEditorMode/Private/UI/SScriptableToolGroupSetPicker.cpp
+++ b/Engine/Plugins/Editor/ScriptableToolsEditorMode/Source/ScriptableToolsEditorMode/Private/UI/SScriptableToolGroupSetPicker.cpp
@@ -249,7 +249,7 @@ void SScriptableToolGroupSetPicker::PopulateCheckedTags()
 	if (bValidData)
 	{
 		FScriptableToolGroupSet& GroupSet = *static_cast<FScriptableToolGroupSet*>(StructPointer);
-		FScriptableToolGroupSet::FGroupSet& Groups = GroupSet.Groups;
+		FScriptableToolGroupSet::FGroupSet Groups = GroupSet.GetGroups();
 
 		HelperGroupSet->SetGroups(Groups);
 
@@ -324,7 +324,7 @@ void SScriptableToolGroupSetPicker::FlushHelperGroupSet() const
 	}
 	else if (StructPtr)
 	{
-		StructPtr->Groups = HelperGroupSet->GetGroups();
+		StructPtr->SetGroups(HelperGroupSet->GetGroups());
 	}
 
 	OnChanged.ExecuteIfBound();
diff --git a/Engine/Plugins/Editor/ScriptableToolsEditorMode/Source/ScriptableToolsEditorMode/Public/ScriptableToolsEditorModeSettings.h b/Engine/Plugins/Editor/ScriptableToolsEditorMode/Source/ScriptableToolsEditorMode/Public/ScriptableToolsEditorModeSettings.h
--- a/Engine/Plugins/Editor/ScriptableToolsEditorMode/Source/ScriptableToolsEditorMode/Public/ScriptableToolsEditorModeSettings.h
+++ b/Engine/Plugins/Editor/ScriptableToolsEditorMode/Source/ScriptableToolsEditorMode/Public/ScriptableToolsEditorModeSettings.h
@@ -29,6 +29,6 @@ public:
 	UPROPERTY(config, EditAnywhere, Category = "Scriptable Tools Mode|Tool Registration", meta = (EditCondition = "!bRegisterAllTools"))
 	FScriptableToolGroupSet ToolRegistrationFilters;
 
-	bool RegisterAllTools() const {	return ToolRegistrationFilters.Groups.IsEmpty(); }
+	bool RegisterAllTools() const {	return ToolRegistrationFilters.GetGroups().IsEmpty(); }
 
 };
\ No newline at end of file
diff --git a/Engine/Plugins/Runtime/ScriptableToolsFramework/Source/EditorScriptableToolsFramework/EditorScriptableToolsFramework.Build.cs b/Engine/Plugins/Runtime/ScriptableToolsFramework/Source/EditorScriptableToolsFramework/EditorScriptableToolsFramework.Build.cs
--- a/Engine/Plugins/Runtime/ScriptableToolsFramework/Source/EditorScriptableToolsFramework/EditorScriptableToolsFramework.Build.cs
+++ b/Engine/Plugins/Runtime/ScriptableToolsFramework/Source/EditorScriptableToolsFramework/EditorScriptableToolsFramework.Build.cs
@@ -39,7 +39,8 @@ public class EditorScriptableToolsFramework : ModuleRules
 		PrivateDependencyModuleNames.AddRange(
 			new string[]
 			{
-				"Engine"
+				"Engine",
+				"UnrealEd"
 			}
 			);
 
diff --git a/Engine/Plugins/Runtime/ScriptableToolsFramework/Source/EditorScriptableToolsFramework/Private/ScriptableToolSet.cpp b/Engine/Plugins/Runtime/ScriptableToolsFramework/Source/EditorScriptableToolsFramework/Private/ScriptableToolSet.cpp
--- a/Engine/Plugins/Runtime/ScriptableToolsFramework/Source/EditorScriptableToolsFramework/Private/ScriptableToolSet.cpp
+++ b/Engine/Plugins/Runtime/ScriptableToolsFramework/Source/EditorScriptableToolsFramework/Private/ScriptableToolSet.cpp
@@ -7,10 +7,41 @@
 #include "Tags/ScriptableToolGroupSet.h"
 #include "Engine/AssetManager.h" // Singleton access to StreamableManager
 
+#include "Utility/ScriptableToolLogging.h"
 
 #include "Modules/ModuleManager.h"
 #include "AssetRegistry/AssetRegistryModule.h"
 
+#if WITH_EDITOR
+	#include "Editor.h"
+#endif
+
+UScriptableToolSet::UScriptableToolSet()
+{
+	AssetCanDeleteHandle = FEditorDelegates::OnAssetsCanDelete.AddUObject(this, &UScriptableToolSet::HandleAssetCanDelete);
+}
+
+UScriptableToolSet::~UScriptableToolSet()
+{
+	FEditorDelegates::OnAssetsCanDelete.Remove(AssetCanDeleteHandle);
+}
+
+void UScriptableToolSet::HandleAssetCanDelete(const TArray<UObject*>& InObjectsToDelete, FCanDeleteAssetResult& OutCanDelete)
+{
+	OutCanDelete.Set(true);
+
+	for (UObject* Obj : InObjectsToDelete)
+	{
+		if (Tools.ContainsByPredicate([&Obj](const FScriptableToolInfo& ToolInfo) {
+			return ToolInfo.ToolPath == Obj->GetPathName() || ToolInfo.BuilderPath == Obj->GetPathName();
+			}))
+		{
+			UE_LOG(LogScriptableTools, Error, TEXT("Unable to delete asset. Asset is currently loaded as part of current scriptable tool palette: %s"), *Obj->GetPathName());
+
+			OutCanDelete.Set(false);
+		}
+	}
+}
 
 void UScriptableToolSet::ReinitializeScriptableTools(FPreToolsLoadedDelegate PreDelegate, FToolsLoadedDelegate PostDelegate, FToolsLoadingUpdateDelegate UpdateDelegate, FScriptableToolGroupSet* TagsToFilter)
 {
@@ -138,6 +169,7 @@ void UScriptableToolSet::PostToolLoad(FToolsLoadedDelegate Delegate, TArray< FSo
 		if (Class->IsChildOf(ScriptableToolClass))
 		{
 			FScriptableToolInfo ToolInfo;
+			ToolInfo.ToolPath = Class->GetPathName().LeftChop(2); // Remove the _C postfix - this lets us match against the asset later and we don't actually need it otherwise.
 			ToolInfo.ToolClass = Class;
 			ToolInfo.ToolCDO = Class->GetDefaultObject<UScriptableInteractiveTool>();
 
@@ -150,6 +182,22 @@ void UScriptableToolSet::PostToolLoad(FToolsLoadedDelegate Delegate, TArray< FSo
 			ToolBuilder->ToolClass = Class;
 			ToolBuilders.Add(ToolBuilder);
 
+			switch (ToolInfo.ToolCDO->ToolStartupRequirements)
+			{
+				// We need to get the path to the specific tool builder the tool is using. Despite appearances, this is how you do it.
+				// Trying to do this directly on the ToolBuilder variable will result in the wrong path.
+				case EScriptableToolStartupRequirements::ToolTarget:
+				{					
+					ToolInfo.BuilderPath = ToolInfo.ToolCDO->ToolTargetToolBuilderClass->GetPathName().LeftChop(2);
+					break;
+				}
+				case EScriptableToolStartupRequirements::Custom:
+				{
+					ToolInfo.BuilderPath = ToolInfo.ToolCDO->CustomToolBuilderClass->GetPathName().LeftChop(2);
+					break;
+				}
+			}
+			
 			ToolInfo.ToolBuilder = ToolBuilder;
 			Tools.Add(ToolInfo);
 		}
diff --git a/Engine/Plugins/Runtime/ScriptableToolsFramework/Source/EditorScriptableToolsFramework/Public/ScriptableToolSet.h b/Engine/Plugins/Runtime/ScriptableToolsFramework/Source/EditorScriptableToolsFramework/Public/ScriptableToolSet.h
--- a/Engine/Plugins/Runtime/ScriptableToolsFramework/Source/EditorScriptableToolsFramework/Public/ScriptableToolSet.h
+++ b/Engine/Plugins/Runtime/ScriptableToolsFramework/Source/EditorScriptableToolsFramework/Public/ScriptableToolSet.h
@@ -9,6 +9,7 @@ class UScriptableInteractiveTool;
 class UBaseScriptableToolBuilder;
 struct FScriptableToolGroupSet;
 class UClass;
+struct FCanDeleteAssetResult;
 
 DECLARE_DELEGATE(FPreToolsLoadedDelegate);
 DECLARE_DELEGATE(FToolsLoadedDelegate);
@@ -23,6 +24,10 @@ class EDITORSCRIPTABLETOOLSFRAMEWORK_API UScriptableToolSet : public UObject
 
 public:
 
+	UScriptableToolSet();
+
+	virtual ~UScriptableToolSet();
+
 	/**
 	 * Find all UScriptableInteractiveTool classes in the current project.
 	 * (Currently no support for filtering/etc)
@@ -35,15 +40,21 @@ public:
 	void ForEachScriptableTool(
 		TFunctionRef<void(UClass* ToolClass, UBaseScriptableToolBuilder* ToolBuilder)> ProcessToolFunc);
 
-protected:
+private:
+
+	void HandleAssetCanDelete(const TArray<UObject*>& InObjectsToDelete, FCanDeleteAssetResult& OutCanDelete);
 
 	void PostToolLoad(FToolsLoadedDelegate Delegate, TArray< FSoftObjectPath > ObjectsLoaded, TSharedPtr<FScriptableToolGroupSet> TagsToFilter);
 
 	bool bActiveLoading = false;
 	TSharedPtr<FStreamableHandle> AsyncLoadHandle;
 
+	FDelegateHandle AssetCanDeleteHandle;
+
 	struct FScriptableToolInfo
 	{
+		FString ToolPath;
+		FString BuilderPath;
 		TWeakObjectPtr<UClass> ToolClass = nullptr;
 		TWeakObjectPtr<UScriptableInteractiveTool> ToolCDO;
 		TWeakObjectPtr<UBaseScriptableToolBuilder> ToolBuilder;
diff --git a/Engine/Plugins/Runtime/ScriptableToolsFramework/Source/ScriptableToolsFramework/Private/Tags/ScriptableToolGroupSet.cpp b/Engine/Plugins/Runtime/ScriptableToolsFramework/Source/ScriptableToolsFramework/Private/Tags/ScriptableToolGroupSet.cpp
--- a/Engine/Plugins/Runtime/ScriptableToolsFramework/Source/ScriptableToolsFramework/Private/Tags/ScriptableToolGroupSet.cpp
+++ b/Engine/Plugins/Runtime/ScriptableToolsFramework/Source/ScriptableToolsFramework/Private/Tags/ScriptableToolGroupSet.cpp
@@ -2,3 +2,10 @@
 
 #include "Tags/ScriptableToolGroupSet.h"
 
+void FScriptableToolGroupSet::SanitizeGroups()
+{
+	FGroupSet GroupSetCopy = Groups;
+	GroupSetCopy.Remove(nullptr);
+	Groups = GroupSetCopy;
+}
+
diff --git a/Engine/Plugins/Runtime/ScriptableToolsFramework/Source/ScriptableToolsFramework/Private/Utility/ScriptableToolLogging.cpp b/Engine/Plugins/Runtime/ScriptableToolsFramework/Source/ScriptableToolsFramework/Private/Utility/ScriptableToolLogging.cpp
new file mode 100644
--- /dev/null
+++ b/Engine/Plugins/Runtime/ScriptableToolsFramework/Source/ScriptableToolsFramework/Private/Utility/ScriptableToolLogging.cpp
@@ -0,0 +1,5 @@
+// Copyright Epic Games, Inc. All Rights Reserved.
+
+#include "Utility/ScriptableToolLogging.h"
+
+DEFINE_LOG_CATEGORY(LogScriptableTools)
\ No newline at end of file
diff --git a/Engine/Plugins/Runtime/ScriptableToolsFramework/Source/ScriptableToolsFramework/Public/ScriptableToolBuilder.h b/Engine/Plugins/Runtime/ScriptableToolsFramework/Source/ScriptableToolsFramework/Public/ScriptableToolBuilder.h
--- a/Engine/Plugins/Runtime/ScriptableToolsFramework/Source/ScriptableToolsFramework/Public/ScriptableToolBuilder.h
+++ b/Engine/Plugins/Runtime/ScriptableToolsFramework/Source/ScriptableToolsFramework/Public/ScriptableToolBuilder.h
@@ -133,7 +133,7 @@ private:
 };
 
 UCLASS(Transient, Blueprintable, Abstract)
-class UToolTargetScriptableToolBuilder: public UCustomScriptableToolBuilderComponentBase, public ICustomScriptableToolBuilderBaseInterface
+class SCRIPTABLETOOLSFRAMEWORK_API UToolTargetScriptableToolBuilder: public UCustomScriptableToolBuilderComponentBase, public ICustomScriptableToolBuilderBaseInterface
 {
 	GENERATED_BODY()
 
diff --git a/Engine/Plugins/Runtime/ScriptableToolsFramework/Source/ScriptableToolsFramework/Public/Tags/ScriptableToolGroupSet.h b/Engine/Plugins/Runtime/ScriptableToolsFramework/Source/ScriptableToolsFramework/Public/Tags/ScriptableToolGroupSet.h
--- a/Engine/Plugins/Runtime/ScriptableToolsFramework/Source/ScriptableToolsFramework/Public/Tags/ScriptableToolGroupSet.h
+++ b/Engine/Plugins/Runtime/ScriptableToolsFramework/Source/ScriptableToolsFramework/Public/Tags/ScriptableToolGroupSet.h
@@ -9,7 +9,7 @@
 #include "ScriptableToolGroupSet.generated.h"
 
 USTRUCT()
-struct FScriptableToolGroupSet
+struct SCRIPTABLETOOLSFRAMEWORK_API FScriptableToolGroupSet
 {
 	GENERATED_BODY()
 
@@ -18,13 +18,33 @@ struct FScriptableToolGroupSet
 	 */
 	typedef TSet<TSubclassOf<UScriptableToolGroupTag>> FGroupSet;
 
-	UPROPERTY(EditAnywhere, Category = "Groups")
-	TSet<TSubclassOf<UScriptableToolGroupTag>> Groups;
-
 public:
 
 	bool Matches(const FScriptableToolGroupSet& OtherSet) const
 	{
-		return !Groups.Intersect(OtherSet.Groups).IsEmpty();
+		return !GetGroups().Intersect(OtherSet.GetGroups()).IsEmpty();
+	}
+
+	void SetGroups(const FGroupSet& GroupsIn)
+	{
+		Groups = GroupsIn;
+		SanitizeGroups();
+	}
+
+	const FGroupSet& GetGroups() const
+	{
+		return Groups;
 	}
+
+	FGroupSet& GetGroups()
+	{
+		return Groups;
+	}
+
+private:
+
+	void  SanitizeGroups();
+
+	UPROPERTY()
+	TSet<TSubclassOf<UScriptableToolGroupTag>> Groups;
 };
\ No newline at end of file
diff --git a/Engine/Plugins/Runtime/ScriptableToolsFramework/Source/ScriptableToolsFramework/Public/Utility/ScriptableToolLogging.h b/Engine/Plugins/Runtime/ScriptableToolsFramework/Source/ScriptableToolsFramework/Public/Utility/ScriptableToolLogging.h
new file mode 100644
--- /dev/null
+++ b/Engine/Plugins/Runtime/ScriptableToolsFramework/Source/ScriptableToolsFramework/Public/Utility/ScriptableToolLogging.h
@@ -0,0 +1,7 @@
+// Copyright Epic Games, Inc. All Rights Reserved.
+
+#pragma once
+
+#include "CoreMinimal.h"
+
+SCRIPTABLETOOLSFRAMEWORK_API DECLARE_LOG_CATEGORY_EXTERN(LogScriptableTools, Log, All)
\ No newline at end of file
