ScriptableTools: Initialize new triangles and quad geometry with a default material to avoid spurious ensure failures when removing them before a user manually sets a material.

diff --git a/Engine/Plugins/Runtime/ScriptableToolsFramework/Source/ScriptableToolsFramework/Private/Drawing/ScriptableToolTriangle.cpp b/Engine/Plugins/Runtime/ScriptableToolsFramework/Source/ScriptableToolsFramework/Private/Drawing/ScriptableToolTriangle.cpp
--- a/Engine/Plugins/Runtime/ScriptableToolsFramework/Source/ScriptableToolsFramework/Private/Drawing/ScriptableToolTriangle.cpp
+++ b/Engine/Plugins/Runtime/ScriptableToolsFramework/Source/ScriptableToolsFramework/Private/Drawing/ScriptableToolTriangle.cpp
@@ -1,7 +1,13 @@
 // Copyright Epic Games, Inc. All Rights Reserved.
 
 #include "Drawing/ScriptableToolTriangle.h"
+#include "Materials/Material.h"
+#include "MaterialDomain.h"
 
+UScriptableToolTriangle::UScriptableToolTriangle()
+{
+	TriangleDescription.Material = UMaterial::GetDefaultMaterial(EMaterialDomain::MD_Surface);
+}
 
 void UScriptableToolTriangle::SetTriangleID(int32 TriangleIDIn)
 {
@@ -65,7 +71,11 @@ void UScriptableToolTriangle::SetTriangleColors(FColor A, FColor B, FColor C)
 
 
 
-
+UScriptableToolQuad::UScriptableToolQuad()
+{
+	TriangleADescription.Material = UMaterial::GetDefaultMaterial(EMaterialDomain::MD_Surface);
+	TriangleBDescription.Material = UMaterial::GetDefaultMaterial(EMaterialDomain::MD_Surface);
+}
 
 
 void UScriptableToolQuad::SetTriangleAID(int32 TriangleIDIn)
diff --git a/Engine/Plugins/Runtime/ScriptableToolsFramework/Source/ScriptableToolsFramework/Public/Drawing/ScriptableToolTriangle.h b/Engine/Plugins/Runtime/ScriptableToolsFramework/Source/ScriptableToolsFramework/Public/Drawing/ScriptableToolTriangle.h
--- a/Engine/Plugins/Runtime/ScriptableToolsFramework/Source/ScriptableToolsFramework/Public/Drawing/ScriptableToolTriangle.h
+++ b/Engine/Plugins/Runtime/ScriptableToolsFramework/Source/ScriptableToolsFramework/Public/Drawing/ScriptableToolTriangle.h
@@ -17,7 +17,7 @@ class SCRIPTABLETOOLSFRAMEWORK_API UScriptableToolTriangle : public UObject
 
 public:
 
-	//void Initialize();
+	UScriptableToolTriangle();
 
 	void SetTriangleID(int32 TriangleIDIn);
 	int32 GetTriangleID() const;
@@ -93,6 +93,8 @@ class SCRIPTABLETOOLSFRAMEWORK_API UScriptableToolQuad : public UObject
 
 public:
 
+	UScriptableToolQuad();
+
 	void SetTriangleAID(int32 TriangleIDIn);
 	int32 GetTriangleAID() const;
 	void SetTriangleBID(int32 TriangleIDIn);
