ScriptableTools: Improve handling of mode shutdown to properly release references to tools allowing deletion of blueprint assets.

diff --git a/Engine/Plugins/Editor/ScriptableToolsEditorMode/Source/ScriptableToolsEditorMode/Private/ScriptableToolsEditorMode.cpp b/Engine/Plugins/Editor/ScriptableToolsEditorMode/Source/ScriptableToolsEditorMode/Private/ScriptableToolsEditorMode.cpp
--- a/Engine/Plugins/Editor/ScriptableToolsEditorMode/Source/ScriptableToolsEditorMode/Private/ScriptableToolsEditorMode.cpp
+++ b/Engine/Plugins/Editor/ScriptableToolsEditorMode/Source/ScriptableToolsEditorMode/Private/ScriptableToolsEditorMode.cpp
@@ -342,6 +342,10 @@ void UScriptableToolsEditorMode::Exit()
 		}
 	}
 
+	// Explicitly unload all tools from the set, just in case
+	ScriptableTools->UnloadAllTools();
+	ScriptableTools = nullptr;
+
 	// Call base Exit method to ensure proper cleanup
 	UEdMode::Exit();
 }
diff --git a/Engine/Plugins/Runtime/ScriptableToolsFramework/Source/EditorScriptableToolsFramework/Private/ScriptableToolSet.cpp b/Engine/Plugins/Runtime/ScriptableToolsFramework/Source/EditorScriptableToolsFramework/Private/ScriptableToolSet.cpp
--- a/Engine/Plugins/Runtime/ScriptableToolsFramework/Source/EditorScriptableToolsFramework/Private/ScriptableToolSet.cpp
+++ b/Engine/Plugins/Runtime/ScriptableToolsFramework/Source/EditorScriptableToolsFramework/Private/ScriptableToolSet.cpp
@@ -24,6 +24,21 @@ UScriptableToolSet::UScriptableToolSet()
 UScriptableToolSet::~UScriptableToolSet()
 {
 	FEditorDelegates::OnAssetsCanDelete.Remove(AssetCanDeleteHandle);
+
+	UnloadAllTools();
+}
+
+void  UScriptableToolSet::UnloadAllTools()
+{
+	if (bActiveLoading)
+	{
+		AsyncLoadHandle->CancelHandle();
+	}
+
+	Tools.Reset();
+	ToolBuilders.Reset();
+
+	bActiveLoading = false;
 }
 
 void UScriptableToolSet::HandleAssetCanDelete(const TArray<UObject*>& InObjectsToDelete, FCanDeleteAssetResult& OutCanDelete)
@@ -173,6 +188,12 @@ void UScriptableToolSet::PostToolLoad(FToolsLoadedDelegate Delegate, TArray< FSo
 			ToolInfo.ToolClass = Class;
 			ToolInfo.ToolCDO = Class->GetDefaultObject<UScriptableInteractiveTool>();
 
+			if(!ToolInfo.ToolCDO->bShowToolInEditor)
+			{
+				continue; // If the tool doesn't want to be shown in editor, we can skip loading it entirely,
+				          // preventing the deletion prevention from kicking in and spending cycles on tool builder tests
+			}
+
 			UBaseScriptableToolBuilder* ToolBuilder = ToolInfo.ToolCDO->GetNewCustomToolBuilderInstance(this);
 			if (ToolBuilder == nullptr)
 			{
@@ -198,7 +219,8 @@ void UScriptableToolSet::PostToolLoad(FToolsLoadedDelegate Delegate, TArray< FSo
 				}
 			}
 			
-			ToolInfo.ToolBuilder = ToolBuilder;
+			ToolInfo.ToolBuilder = ToolBuilder;			
+
 			Tools.Add(ToolInfo);
 		}
 	}
diff --git a/Engine/Plugins/Runtime/ScriptableToolsFramework/Source/EditorScriptableToolsFramework/Public/ScriptableToolSet.h b/Engine/Plugins/Runtime/ScriptableToolsFramework/Source/EditorScriptableToolsFramework/Public/ScriptableToolSet.h
--- a/Engine/Plugins/Runtime/ScriptableToolsFramework/Source/EditorScriptableToolsFramework/Public/ScriptableToolSet.h
+++ b/Engine/Plugins/Runtime/ScriptableToolsFramework/Source/EditorScriptableToolsFramework/Public/ScriptableToolSet.h
@@ -28,6 +28,11 @@ public:
 
 	virtual ~UScriptableToolSet();
 
+	/**
+	* Forces the unloading of all tools loaded
+	*/
+	void UnloadAllTools();
+
 	/**
 	 * Find all UScriptableInteractiveTool classes in the current project.
 	 * (Currently no support for filtering/etc)
