Fixed usage of TRACE_CPUPROFILER_EVENT_SCOPE macro in various files.

#rnx

diff --git a/Engine/Plugins/Runtime/ScriptableToolsFramework/Source/EditorScriptableToolsFramework/Private/ScriptableToolSet.cpp b/Engine/Plugins/Runtime/ScriptableToolsFramework/Source/EditorScriptableToolsFramework/Private/ScriptableToolSet.cpp
--- a/Engine/Plugins/Runtime/ScriptableToolsFramework/Source/EditorScriptableToolsFramework/Private/ScriptableToolSet.cpp
+++ b/Engine/Plugins/Runtime/ScriptableToolsFramework/Source/EditorScriptableToolsFramework/Private/ScriptableToolSet.cpp
@@ -60,7 +60,7 @@ void UScriptableToolSet::HandleAssetCanDelete(const TArray<UObject*>& InObjectsT
 
 void UScriptableToolSet::ReinitializeScriptableTools(FPreToolsLoadedDelegate PreDelegate, FToolsLoadedDelegate PostDelegate, FToolsLoadingUpdateDelegate UpdateDelegate, FScriptableToolGroupSet* TagsToFilter)
 {
-	TRACE_CPUPROFILER_EVENT_SCOPE(TEXT("UScriptableToolSet::ReinitializeScriptableTools"))
+	TRACE_CPUPROFILER_EVENT_SCOPE(UScriptableToolSet::ReinitializeScriptableTools)
 
 	if (bActiveLoading)
 	{
@@ -151,7 +151,7 @@ void UScriptableToolSet::ForEachScriptableTool(
 
 void UScriptableToolSet::PostToolLoad(FToolsLoadedDelegate Delegate, TArray< FSoftObjectPath > ObjectsLoaded, TSharedPtr<FScriptableToolGroupSet> TagsToFilter)
 {
-	TRACE_CPUPROFILER_EVENT_SCOPE(TEXT("UScriptableToolSet::PostToolLoad"))
+	TRACE_CPUPROFILER_EVENT_SCOPE(UScriptableToolSet::PostToolLoad)
 
 	TSet<UClass*> PotentialToolClasses;
 	UClass* ScriptableToolClass = UScriptableInteractiveTool::StaticClass();
diff --git a/Engine/Source/Runtime/Experimental/Chaos/Public/Chaos/XPBDCorotatedConstraints.h b/Engine/Source/Runtime/Experimental/Chaos/Public/Chaos/XPBDCorotatedConstraints.h
--- a/Engine/Source/Runtime/Experimental/Chaos/Public/Chaos/XPBDCorotatedConstraints.h
+++ b/Engine/Source/Runtime/Experimental/Chaos/Public/Chaos/XPBDCorotatedConstraints.h
@@ -408,7 +408,7 @@ namespace Chaos::Softs
 
 		void ModifyDmInverseFromFiberLength(const int32 e, const T L, const PMatrix<T, 3, 3>& MFiberDir, const T ContractionVolumeScale) const
 		{
-			TRACE_CPUPROFILER_EVENT_SCOPE(TEXT("STAT_ChaosXPBDCorotatedModifyDmInverseFromFiberLength"));
+			TRACE_CPUPROFILER_EVENT_SCOPE(STAT_ChaosXPBDCorotatedModifyDmInverseFromFiberLength);
 			PMatrix<T, 3, 3> DmInv = ElementDmInvSave(e);
 			if (L < 1)
 			{
diff --git a/Engine/Source/Runtime/SlateCore/Private/Rendering/DrawElements.cpp b/Engine/Source/Runtime/SlateCore/Private/Rendering/DrawElements.cpp
--- a/Engine/Source/Runtime/SlateCore/Private/Rendering/DrawElements.cpp
+++ b/Engine/Source/Runtime/SlateCore/Private/Rendering/DrawElements.cpp
@@ -81,7 +81,7 @@ void FSlateWindowElementList::StartMergeRenderBatches()
 	{
 		MergeBatchDataTask = UE::Tasks::Launch(TEXT("Slate::MergeRenderBatches"), [this]
 		{
-			TRACE_CPUPROFILER_EVENT_SCOPE("Slate::MergeRenderBatches")
+			TRACE_CPUPROFILER_EVENT_SCOPE(Slate::MergeRenderBatches);
 			BatchData.MergeRenderBatches();
 			BatchDataHDR.MergeRenderBatches();
 		});
diff --git a/Engine/Source/Runtime/StaticMeshDescription/Private/StaticMeshOperations.cpp b/Engine/Source/Runtime/StaticMeshDescription/Private/StaticMeshOperations.cpp
--- a/Engine/Source/Runtime/StaticMeshDescription/Private/StaticMeshOperations.cpp
+++ b/Engine/Source/Runtime/StaticMeshDescription/Private/StaticMeshOperations.cpp
@@ -1038,7 +1038,7 @@ void FStaticMeshOperations::ConvertFromRawMesh(const FRawMesh& SourceRawMesh, FM
 
 void FStaticMeshOperations::AppendMeshDescriptions(const TArray<const FMeshDescription*>& SourceMeshes, FMeshDescription& TargetMesh, FAppendSettings& AppendSettings)
 {
-	TRACE_CPUPROFILER_EVENT_SCOPE_STR((SourceMeshes.Num() > 1 ? "FStaticMeshOperations::AppendMeshDescriptions" : "FStaticMeshOperations::AppendMeshDescription"));
+	TRACE_CPUPROFILER_EVENT_SCOPE(FStaticMeshOperations::AppendMeshDescriptions);
 
 	FStaticMeshAttributes TargetAttributes(TargetMesh);
 	TVertexAttributesRef<FVector3f> TargetVertexPositions = TargetAttributes.GetVertexPositions();
