ScriptableTools: Make the gizmo respect the "scale uniform" flag, and add the free translate and rotate flags so that uniform scale, arcball, and screen translation can be used.

diff --git a/Engine/Plugins/Runtime/ScriptableToolsFramework/Source/ScriptableToolsFramework/Private/ScriptableInteractiveTool.cpp b/Engine/Plugins/Runtime/ScriptableToolsFramework/Source/ScriptableToolsFramework/Private/ScriptableInteractiveTool.cpp
--- a/Engine/Plugins/Runtime/ScriptableToolsFramework/Source/ScriptableToolsFramework/Private/ScriptableInteractiveTool.cpp
+++ b/Engine/Plugins/Runtime/ScriptableToolsFramework/Source/ScriptableToolsFramework/Private/ScriptableInteractiveTool.cpp
@@ -938,6 +938,7 @@ void UScriptableInteractiveTool::CreateTRSGizmo(
 	ConfigureGizmoTranslation(EScriptableToolGizmoTranslation::TranslatePlaneXY, ETransformGizmoSubElements::TranslatePlaneXY);
 	ConfigureGizmoTranslation(EScriptableToolGizmoTranslation::TranslatePlaneXZ, ETransformGizmoSubElements::TranslatePlaneXZ);
 	ConfigureGizmoTranslation(EScriptableToolGizmoTranslation::TranslatePlaneYZ, ETransformGizmoSubElements::TranslatePlaneYZ);
+	ConfigureGizmoTranslation(EScriptableToolGizmoTranslation::FreeTranslate, ETransformGizmoSubElements::FreeTranslate);
 
 	auto ConfigureGizmoRotation = [&SubElements, &GizmoOptions](EScriptableToolGizmoRotation SettingFlag, ETransformGizmoSubElements GizmoFlag)
 	{
@@ -946,6 +947,7 @@ void UScriptableInteractiveTool::CreateTRSGizmo(
 	ConfigureGizmoRotation(EScriptableToolGizmoRotation::RotateAxisX, ETransformGizmoSubElements::RotateAxisX);
 	ConfigureGizmoRotation(EScriptableToolGizmoRotation::RotateAxisY, ETransformGizmoSubElements::RotateAxisY);
 	ConfigureGizmoRotation(EScriptableToolGizmoRotation::RotateAxisZ, ETransformGizmoSubElements::RotateAxisZ);
+	ConfigureGizmoRotation(EScriptableToolGizmoRotation::FreeRotate, ETransformGizmoSubElements::FreeRotate);
 
 	auto ConfigureGizmoScale = [&SubElements, &GizmoOptions](EScriptableToolGizmoScale SettingFlag, ETransformGizmoSubElements GizmoFlag)
 	{
@@ -957,6 +959,7 @@ void UScriptableInteractiveTool::CreateTRSGizmo(
 	ConfigureGizmoScale(EScriptableToolGizmoScale::ScalePlaneXY, ETransformGizmoSubElements::ScalePlaneXY);
 	ConfigureGizmoScale(EScriptableToolGizmoScale::ScalePlaneXZ, ETransformGizmoSubElements::ScalePlaneXZ);
 	ConfigureGizmoScale(EScriptableToolGizmoScale::ScalePlaneYZ, ETransformGizmoSubElements::ScalePlaneYZ);
+	ConfigureGizmoScale(EScriptableToolGizmoScale::ScaleUniform, ETransformGizmoSubElements::ScaleUniform);
 
 	if (SubElements == ETransformGizmoSubElements::None)
 	{
diff --git a/Engine/Plugins/Runtime/ScriptableToolsFramework/Source/ScriptableToolsFramework/Public/ScriptableInteractiveTool.h b/Engine/Plugins/Runtime/ScriptableToolsFramework/Source/ScriptableToolsFramework/Public/ScriptableInteractiveTool.h
--- a/Engine/Plugins/Runtime/ScriptableToolsFramework/Source/ScriptableToolsFramework/Public/ScriptableInteractiveTool.h
+++ b/Engine/Plugins/Runtime/ScriptableToolsFramework/Source/ScriptableToolsFramework/Public/ScriptableInteractiveTool.h
@@ -95,6 +95,7 @@ enum class EScriptableToolGizmoTranslation : uint8
 	TranslatePlaneXY = 1 << 4,
 	TranslatePlaneXZ = 1 << 5,
 	TranslatePlaneYZ = 1 << 6,
+	FreeTranslate = 1 << 7,
 
 	All = 0xFF UMETA(Hidden)
 };
@@ -110,6 +111,8 @@ enum class EScriptableToolGizmoRotation : uint8
 	RotateAxisY = 1 << 2,
 	RotateAxisZ = 1 << 3,
 
+	FreeRotate = 1 << 7,
+
 	All = 0xFF UMETA(Hidden)
 };
 ENUM_CLASS_FLAGS(EScriptableToolGizmoRotation);
