ScriptableTools: Added GetWorld to ScriptableInteractiveToolPropertySet so that WorldContextObject dependent BP functions are accessible.

diff --git a/Engine/Plugins/Runtime/ScriptableToolsFramework/Source/ScriptableToolsFramework/Private/ScriptableInteractiveTool.cpp b/Engine/Plugins/Runtime/ScriptableToolsFramework/Source/ScriptableToolsFramework/Private/ScriptableInteractiveTool.cpp
--- a/Engine/Plugins/Runtime/ScriptableToolsFramework/Source/ScriptableToolsFramework/Private/ScriptableInteractiveTool.cpp
+++ b/Engine/Plugins/Runtime/ScriptableToolsFramework/Source/ScriptableToolsFramework/Private/ScriptableInteractiveTool.cpp
@@ -41,17 +41,24 @@ namespace ScriptableInteractiveToolLocals
 
 UScriptableInteractiveTool* UScriptableInteractiveToolPropertySet::GetOwningTool(EToolsFrameworkOutcomePins& Outcome)
 {
-	if (ParentTool.IsValid())
+	if (UScriptableInteractiveTool* OwningTool = ParentTool.Get())
 	{
 		Outcome = EToolsFrameworkOutcomePins::Success;
-		return ParentTool.Get();
+		return OwningTool;
 	}
 	Outcome = EToolsFrameworkOutcomePins::Failure;
 	return nullptr;
 }
 
-
-
+UWorld* UScriptableInteractiveToolPropertySet::GetWorld() const
+{
+	UWorld* Result = nullptr;
+	if (UScriptableInteractiveTool* OwningTool = ParentTool.Get())
+	{
+		Result = OwningTool->GetWorld();
+	}
+	return Result;
+}
 
 UScriptableTool_RenderAPI* UScriptableTool_RenderAPI::DrawLine(FVector Start, FVector End, FLinearColor Color, float Thickness, float DepthBias, bool bDepthTested)
 {
diff --git a/Engine/Plugins/Runtime/ScriptableToolsFramework/Source/ScriptableToolsFramework/Public/ScriptableInteractiveTool.h b/Engine/Plugins/Runtime/ScriptableToolsFramework/Source/ScriptableToolsFramework/Public/ScriptableInteractiveTool.h
--- a/Engine/Plugins/Runtime/ScriptableToolsFramework/Source/ScriptableToolsFramework/Public/ScriptableInteractiveTool.h
+++ b/Engine/Plugins/Runtime/ScriptableToolsFramework/Source/ScriptableToolsFramework/Public/ScriptableInteractiveTool.h
@@ -277,6 +277,8 @@ class SCRIPTABLETOOLSFRAMEWORK_API UScriptableInteractiveToolPropertySet : publi
 	UScriptableInteractiveTool* GetOwningTool(
 		EToolsFrameworkOutcomePins& Outcome);
 
+	virtual UWorld* GetWorld() const override;	// must implement this for WorldContextObject dependent BP functions to show.
+
 protected:
 	TWeakObjectPtr<UScriptableInteractiveTool> ParentTool;
 	friend class UScriptableInteractiveTool;
@@ -433,9 +435,9 @@ protected:
 	UPROPERTY(Transient, DuplicateTransient, NonTransactional, SkipSerialization)
 	TWeakObjectPtr<UWorld> TargetWorld;
 
-	virtual UWorld* GetWorld() const override;		// must implement this for some reason, otherwise many BP functions will not show up
-
 public:
+	virtual UWorld* GetWorld() const override;		// must implement this for WorldContextObject dependent BP functions to show.
+	
 	virtual void SetTargetWorld(UWorld* World);		// needs to be accessible to ToolBuilder
 
 	/**
