ScriptableTools: Fix issue where enum properties couldn't be read from integer byte properties correctly.

diff --git a/Engine/Plugins/Runtime/ScriptableToolsFramework/Source/ScriptableToolsFramework/Private/ScriptableInteractiveTool.cpp b/Engine/Plugins/Runtime/ScriptableToolsFramework/Source/ScriptableToolsFramework/Private/ScriptableInteractiveTool.cpp
--- a/Engine/Plugins/Runtime/ScriptableToolsFramework/Source/ScriptableToolsFramework/Private/ScriptableInteractiveTool.cpp
+++ b/Engine/Plugins/Runtime/ScriptableToolsFramework/Source/ScriptableToolsFramework/Private/ScriptableInteractiveTool.cpp
@@ -669,7 +669,23 @@ UScriptableInteractiveToolPropertySet* UScriptableInteractiveTool::WatchEnumProp
 			[this, PropertySet, PropertyName, ArrayIndex](int64 NewValue) { 
 				WatchEnumPropertyDelegates[ArrayIndex].ExecuteIfBound(PropertySet, PropertyName, (uint8)NewValue);
 			} );
+	} else if (FByteProperty* ByteProperty = UELocal::FindValidPropertyByTypeAndName<FByteProperty>(
+		PropertySet, PropertyName, TEXT("WatchEnumProperty")))
+	{
+		int32 ArrayIndex = WatchEnumPropertyDelegates.Num();
+		WatchEnumPropertyDelegates.Add(OnModified);
+
+		PropertySet->WatchProperty<int>(
+			[ByteProperty, PropertySet]()->int64 {
+				const void* ValuePtr = ByteProperty->ContainerPtrToValuePtr<void>(PropertySet);
+				int64 Value = ByteProperty->GetSignedIntPropertyValue(ValuePtr);
+				return Value;
+			},
+			[this, PropertySet, PropertyName, ArrayIndex](int64 NewValue) {
+				WatchEnumPropertyDelegates[ArrayIndex].ExecuteIfBound(PropertySet, PropertyName, (uint8)NewValue);
+			});
 	}
+
 	return PropertySet;
 }
 
