InteractiveToolsFramework: Allow key capture behavior to notify its target of a forced end capture. Expose this in scriptable tools. Also fix a bug with multi key behaviors incorrectly triggering release callbacks after first time all keys were pressed.

diff --git a/Engine/Plugins/Runtime/ScriptableToolsFramework/Source/ScriptableToolsFramework/Private/BaseTools/ScriptableModularBehaviorTool.cpp b/Engine/Plugins/Runtime/ScriptableToolsFramework/Source/ScriptableToolsFramework/Private/BaseTools/ScriptableModularBehaviorTool.cpp
--- a/Engine/Plugins/Runtime/ScriptableToolsFramework/Source/ScriptableToolsFramework/Private/BaseTools/ScriptableModularBehaviorTool.cpp
+++ b/Engine/Plugins/Runtime/ScriptableToolsFramework/Source/ScriptableToolsFramework/Private/BaseTools/ScriptableModularBehaviorTool.cpp
@@ -153,6 +153,7 @@ void UScriptableModularBehaviorTool::AddMouseHoverBehavior(
 void UScriptableModularBehaviorTool::AddSingleKeyInputBehavior(
 	FOnKeyStateToggleDelegate OnKeyPressedFuncIn,
 	FOnKeyStateToggleDelegate OnKeyReleasedFuncIn,
+	FOnForceEndCaptureDelegate_ScriptableTools OnForceEndCaptureFuncIn,
 	FKey Key,
 	const FMouseBehaviorModiferCheckDelegate ModifierCheckFunction,
 	int CapturePriority
@@ -162,7 +163,7 @@ void UScriptableModularBehaviorTool::AddSingleKeyInputBehavior(
 	Keys.Add(Key);
 
 	TObjectPtr<UScriptableToolKeyInputBehavior> BehaviorContainer = NewObject<UScriptableToolKeyInputBehavior>();
-	BehaviorContainer->Init(this, ModifierCheckFunction, OnKeyPressedFuncIn, OnKeyReleasedFuncIn, Keys, true);
+	BehaviorContainer->Init(this, ModifierCheckFunction, OnKeyPressedFuncIn, OnKeyReleasedFuncIn, OnForceEndCaptureFuncIn, Keys, true);
 	BehaviorContainer->SetDefaultPriority(FInputCapturePriority(CapturePriority));
 
 	KeyInputBehaviors.Add(BehaviorContainer);
@@ -172,6 +173,7 @@ void UScriptableModularBehaviorTool::AddSingleKeyInputBehavior(
 void UScriptableModularBehaviorTool::AddMultiKeyInputBehavior(
 	FOnKeyStateToggleDelegate OnKeyPressedFuncIn,
 	FOnKeyStateToggleDelegate OnKeyReleasedFuncIn,
+	FOnForceEndCaptureDelegate_ScriptableTools OnForceEndCaptureFuncIn,
 	TArray<FKey> Keys,
 	bool bRequireAllKeys,
 	const FMouseBehaviorModiferCheckDelegate ModifierCheckFunction,
@@ -179,7 +181,7 @@ void UScriptableModularBehaviorTool::AddMultiKeyInputBehavior(
 )
 {
 	TObjectPtr<UScriptableToolKeyInputBehavior> BehaviorContainer = NewObject<UScriptableToolKeyInputBehavior>();
-	BehaviorContainer->Init(this, ModifierCheckFunction, OnKeyPressedFuncIn, OnKeyReleasedFuncIn, Keys, bRequireAllKeys);
+	BehaviorContainer->Init(this, ModifierCheckFunction, OnKeyPressedFuncIn, OnKeyReleasedFuncIn, OnForceEndCaptureFuncIn, Keys, bRequireAllKeys);
 	BehaviorContainer->SetDefaultPriority(FInputCapturePriority(CapturePriority));
 
 	KeyInputBehaviors.Add(BehaviorContainer);
diff --git a/Engine/Plugins/Runtime/ScriptableToolsFramework/Source/ScriptableToolsFramework/Private/Behaviors/ScriptableToolKeyInputBehavior.cpp b/Engine/Plugins/Runtime/ScriptableToolsFramework/Source/ScriptableToolsFramework/Private/Behaviors/ScriptableToolKeyInputBehavior.cpp
--- a/Engine/Plugins/Runtime/ScriptableToolsFramework/Source/ScriptableToolsFramework/Private/Behaviors/ScriptableToolKeyInputBehavior.cpp
+++ b/Engine/Plugins/Runtime/ScriptableToolsFramework/Source/ScriptableToolsFramework/Private/Behaviors/ScriptableToolKeyInputBehavior.cpp
@@ -10,6 +10,7 @@ void UScriptableToolKeyInputBehavior::Init(
 	FMouseBehaviorModiferCheckDelegate ModifierCheckFuncIn,
 	FOnKeyStateToggleDelegate OnKeyPressedFuncIn,
 	FOnKeyStateToggleDelegate OnKeyReleasedFuncIn,
+	FOnForceEndCaptureDelegate_ScriptableTools OnForceEndCaptureFuncIn,
 	const TArray<FKey>& ListenKeysIn,
 	bool bRequireAllKeysIn
 )
@@ -19,6 +20,7 @@ void UScriptableToolKeyInputBehavior::Init(
 	ModifierCheckFunc = ModifierCheckFuncIn;
 	OnKeyPressedFunc = OnKeyPressedFuncIn;
 	OnKeyReleasedFunc = OnKeyReleasedFuncIn;
+	OnForceEndCaptureFunc = OnForceEndCaptureFuncIn;
 	ListenKeys = ListenKeysIn;
 
 	Behavior->Initialize(this, ListenKeys);
@@ -56,7 +58,6 @@ void UScriptableToolKeyInputBehavior::OnKeyPressed(const FKey& Key)
 	{
 		OnKeyPressedFunc.Execute(Key, BehaviorHost->GetActiveModifiers());
 	}
-	return;
 }
 
 void UScriptableToolKeyInputBehavior::OnKeyReleased(const FKey& Key)
@@ -65,7 +66,14 @@ void UScriptableToolKeyInputBehavior::OnKeyReleased(const FKey& Key)
 	{
 		OnKeyReleasedFunc.Execute(Key, BehaviorHost->GetActiveModifiers());
 	}
-	return;
+}
+
+void UScriptableToolKeyInputBehavior::OnForceEndCapture()
+{
+	if (OnForceEndCaptureFunc.IsBound())
+	{
+		OnForceEndCaptureFunc.Execute();
+	}
 }
 
 void UScriptableToolKeyInputBehavior::OnUpdateModifierState(int ModifierID, bool bIsOn)
diff --git a/Engine/Plugins/Runtime/ScriptableToolsFramework/Source/ScriptableToolsFramework/Private/Behaviors/ScriptableToolKeyInputBehavior.h b/Engine/Plugins/Runtime/ScriptableToolsFramework/Source/ScriptableToolsFramework/Private/Behaviors/ScriptableToolKeyInputBehavior.h
--- a/Engine/Plugins/Runtime/ScriptableToolsFramework/Source/ScriptableToolsFramework/Private/Behaviors/ScriptableToolKeyInputBehavior.h
+++ b/Engine/Plugins/Runtime/ScriptableToolsFramework/Source/ScriptableToolsFramework/Private/Behaviors/ScriptableToolKeyInputBehavior.h
@@ -23,6 +23,7 @@ public:
 		FMouseBehaviorModiferCheckDelegate ModifierCheckFuncIn,
 		FOnKeyStateToggleDelegate OnKeyPressedFuncIn,
 		FOnKeyStateToggleDelegate OnKeyReleasedFuncIn,
+		FOnForceEndCaptureDelegate_ScriptableTools OnForceEndCaptureFuncIn,
 		const TArray<FKey>& ListenKeysIn,
 		bool bRequireAllKeys);
 
@@ -31,6 +32,7 @@ public:
 
 	virtual void OnKeyPressed(const FKey& KeyID) override;
 	virtual void OnKeyReleased(const FKey& KeyID) override;
+	virtual void OnForceEndCapture() override;
 
 	virtual void OnUpdateModifierState(int ModifierID, bool bIsOn) override;
 
@@ -51,6 +53,9 @@ private:
 	UPROPERTY()
 	FOnKeyStateToggleDelegate OnKeyReleasedFunc;
 
+	UPROPERTY()
+	FOnForceEndCaptureDelegate_ScriptableTools OnForceEndCaptureFunc;
+
 	UPROPERTY()
 	TArray<FKey> ListenKeys;
 };
diff --git a/Engine/Plugins/Runtime/ScriptableToolsFramework/Source/ScriptableToolsFramework/Public/BaseTools/ScriptableModularBehaviorTool.h b/Engine/Plugins/Runtime/ScriptableToolsFramework/Source/ScriptableToolsFramework/Public/BaseTools/ScriptableModularBehaviorTool.h
--- a/Engine/Plugins/Runtime/ScriptableToolsFramework/Source/ScriptableToolsFramework/Public/BaseTools/ScriptableModularBehaviorTool.h
+++ b/Engine/Plugins/Runtime/ScriptableToolsFramework/Source/ScriptableToolsFramework/Public/BaseTools/ScriptableModularBehaviorTool.h
@@ -235,6 +235,8 @@ public:
 	*	
 	*	@param OnKeyPressed Callback when the target key is pressed.
 	*	@param OnKeyReleased Callback when the target key is released
+	*	@param OnForceEndCaptureFuncIn Callback when capture is ended prematurely, typically due to the viewport losing
+	*	 focus, in which case the release callback will not be called.
 	*	@param Key,Target key to watch for
 	*	@param CaptureCheck Only enable capture if returns true
 	*	@param CapturePriority The priority is used to resolve situations where multiple behaviors want the same capture
@@ -245,16 +247,30 @@ public:
 	void AddSingleKeyInputBehavior(
 		FOnKeyStateToggleDelegate OnKeyPressed,
 		FOnKeyStateToggleDelegate OnKeyReleased,
+		FOnForceEndCaptureDelegate_ScriptableTools OnForceEndCaptureFuncIn,
 		FKey Key,
 		const FMouseBehaviorModiferCheckDelegate CaptureCheck,
 		int CapturePriority
 	);
 
+	UE_DEPRECATED(5.6, "Use overload that takes an OnForceEndCaptureFunc parameter as well")
+	void AddSingleKeyInputBehavior(
+		FOnKeyStateToggleDelegate OnKeyPressed,
+		FOnKeyStateToggleDelegate OnKeyReleased,
+		FKey Key,
+		const FMouseBehaviorModiferCheckDelegate CaptureCheck,
+		int CapturePriority)
+	{
+		AddSingleKeyInputBehavior(OnKeyPressed, OnKeyReleased, FOnForceEndCaptureDelegate_ScriptableTools(), Key, CaptureCheck, CapturePriority);
+	}
+
 	/**
 	*	AddMultiKeyInputBehavior implements a generic keyboard multi key listener behavior
 	*	
 	*	@param OnKeyPressed  Callback when the target key(s) is pressed. Only triggers once if bRequireAllKeys is true.
 	*	@param OnKeyReleased  Callback when the target key(s) is pressed. Only triggers once if bRequireAllKeys is true.
+	*	@param OnForceEndCaptureFuncIn Callback when capture is ended prematurely, typically due to the viewport losing
+	*	 focus, in which case the release callback will not be called.
 	*	@param Keys Target keys to watch for
 	*	@param bRequireAllKeys If true, all target keys must be pressed simultaniously to recieve press/release events. Otherwise, any and all keys can trigger events.
 	*	@param CaptureCheck Only enable capture if returns true
@@ -266,12 +282,24 @@ public:
 		void AddMultiKeyInputBehavior(
 			FOnKeyStateToggleDelegate OnKeyPressed,
 			FOnKeyStateToggleDelegate OnKeyReleased,
+			FOnForceEndCaptureDelegate_ScriptableTools OnForceEndCaptureFuncIn,
 			TArray<FKey> Keys,
 			bool bRequireAllKeys,
 			const FMouseBehaviorModiferCheckDelegate CaptureCheck,
 			int CapturePriority
 		);
 
+	UE_DEPRECATED(5.6, "Use overload that takes an OnForceEndCaptureFunc parameter as well")
+	void AddMultiKeyInputBehavior(
+		FOnKeyStateToggleDelegate OnKeyPressed,
+		FOnKeyStateToggleDelegate OnKeyReleased,
+		TArray<FKey> Keys,
+		bool bRequireAllKeys,
+		const FMouseBehaviorModiferCheckDelegate CaptureCheck,
+		int CapturePriority)
+	{
+		AddMultiKeyInputBehavior(OnKeyPressed, OnKeyReleased, FOnForceEndCaptureDelegate_ScriptableTools(), Keys, bRequireAllKeys, CaptureCheck, CapturePriority);
+	}
 private:
 
 	UPROPERTY(Transient, DuplicateTransient, NonTransactional, SkipSerialization)
diff --git a/Engine/Plugins/Runtime/ScriptableToolsFramework/Source/ScriptableToolsFramework/Public/Behaviors/ScriptableToolBehaviorDelegates.h b/Engine/Plugins/Runtime/ScriptableToolsFramework/Source/ScriptableToolsFramework/Public/Behaviors/ScriptableToolBehaviorDelegates.h
--- a/Engine/Plugins/Runtime/ScriptableToolsFramework/Source/ScriptableToolsFramework/Public/Behaviors/ScriptableToolBehaviorDelegates.h
+++ b/Engine/Plugins/Runtime/ScriptableToolsFramework/Source/ScriptableToolsFramework/Public/Behaviors/ScriptableToolBehaviorDelegates.h
@@ -44,5 +44,5 @@ DECLARE_DYNAMIC_DELEGATE_RetVal_TwoParams(bool, FOnUpdateHoverDelegate, FInputDe
 DECLARE_DYNAMIC_DELEGATE(FOnEndHoverDelegate);
 
 DECLARE_DYNAMIC_DELEGATE_TwoParams(FOnKeyStateToggleDelegate, FKey, Key, FScriptableToolModifierStates, Modifiers);
-
+DECLARE_DYNAMIC_DELEGATE(FOnForceEndCaptureDelegate_ScriptableTools);
 
diff --git a/Engine/Source/Runtime/InteractiveToolsFramework/Private/BaseBehaviors/KeyInputBehavior.cpp b/Engine/Source/Runtime/InteractiveToolsFramework/Private/BaseBehaviors/KeyInputBehavior.cpp
--- a/Engine/Source/Runtime/InteractiveToolsFramework/Private/BaseBehaviors/KeyInputBehavior.cpp
+++ b/Engine/Source/Runtime/InteractiveToolsFramework/Private/BaseBehaviors/KeyInputBehavior.cpp
@@ -75,6 +75,7 @@ FInputCaptureUpdate UKeyInputBehavior::UpdateCapture(const FInputDeviceState& In
 	if (bRequireAllKeys && (!AreAllKeysPressed() && bAllKeysSeenPressed) )
 	{
 		Target->OnKeyReleased(Input.Keyboard.ActiveKey.Button);
+		bAllKeysSeenPressed = false;
 		return FInputCaptureUpdate::End();
 	}
 
@@ -88,6 +89,10 @@ FInputCaptureUpdate UKeyInputBehavior::UpdateCapture(const FInputDeviceState& In
 
 void UKeyInputBehavior::ForceEndCapture(const FInputCaptureData& Data)
 {
+	if (Target && (bAllKeysSeenPressed || !bRequireAllKeys))
+	{
+		Target->OnForceEndCapture();
+	}
 }
 
 void UKeyInputBehavior::InitializeKeyActivations()
diff --git a/Engine/Source/Runtime/InteractiveToolsFramework/Public/BaseBehaviors/BehaviorTargetInterfaces.h b/Engine/Source/Runtime/InteractiveToolsFramework/Public/BaseBehaviors/BehaviorTargetInterfaces.h
--- a/Engine/Source/Runtime/InteractiveToolsFramework/Public/BaseBehaviors/BehaviorTargetInterfaces.h
+++ b/Engine/Source/Runtime/InteractiveToolsFramework/Public/BaseBehaviors/BehaviorTargetInterfaces.h
@@ -49,6 +49,10 @@ public:
 	virtual void OnKeyReleased(const FKey& KeyID)
 	{
 	}
+
+	virtual void OnForceEndCapture()
+	{
+	}
 };
 
 
