ScriptableTools: Added a GetDrawableGeometryActor() BP function to fetch the actor that owns any line, point and triangle sets added by the tool.

diff --git a/Engine/Plugins/Runtime/ScriptableToolsFramework/Source/ScriptableToolsFramework/Private/ScriptableInteractiveTool.cpp b/Engine/Plugins/Runtime/ScriptableToolsFramework/Source/ScriptableToolsFramework/Private/ScriptableInteractiveTool.cpp
--- a/Engine/Plugins/Runtime/ScriptableToolsFramework/Source/ScriptableToolsFramework/Private/ScriptableInteractiveTool.cpp
+++ b/Engine/Plugins/Runtime/ScriptableToolsFramework/Source/ScriptableToolsFramework/Private/ScriptableInteractiveTool.cpp
@@ -1224,7 +1224,10 @@ TArray<UToolTarget*> UScriptableInteractiveTool::GetToolTargets() const
 }
 
 
-
+AActor* UScriptableInteractiveTool::GetDrawableGeometryActor() const
+{
+	return ToolDrawableGeometry ? ToolDrawableGeometry->GetActor() : nullptr;
+}
 
 UScriptableToolLineSet* UScriptableInteractiveTool::GetDefaultLineSet() const
 {
diff --git a/Engine/Plugins/Runtime/ScriptableToolsFramework/Source/ScriptableToolsFramework/Public/ScriptableInteractiveTool.h b/Engine/Plugins/Runtime/ScriptableToolsFramework/Source/ScriptableToolsFramework/Public/ScriptableInteractiveTool.h
--- a/Engine/Plugins/Runtime/ScriptableToolsFramework/Source/ScriptableToolsFramework/Public/ScriptableInteractiveTool.h
+++ b/Engine/Plugins/Runtime/ScriptableToolsFramework/Source/ScriptableToolsFramework/Public/ScriptableInteractiveTool.h
@@ -821,6 +821,13 @@ public:
 
 	// Drawing API
 
+	/**
+	 * Return the actor that owns any line, point and triangle sets added to this tool.
+	 * @return A reference to the actor object that owns drawable geometry sets.
+	 */
+	UFUNCTION(BlueprintCallable, Category = "ScriptableTool|Drawing")
+	AActor* GetDrawableGeometryActor() const;
+
 	/**
 	 * Retrieve the default line set object for the tool, used for drawing persistent line objects in the scene.
 	 * @return A reference to the tool's default line set object
