ScriptableTools: Added new APIs for property sets to toggle if properties should be readonly or invisible in the editor details view for the tool.

diff --git a/Engine/Plugins/Runtime/ScriptableToolsFramework/Source/ScriptableToolsFramework/Private/ScriptableInteractiveTool.cpp b/Engine/Plugins/Runtime/ScriptableToolsFramework/Source/ScriptableToolsFramework/Private/ScriptableInteractiveTool.cpp
--- a/Engine/Plugins/Runtime/ScriptableToolsFramework/Source/ScriptableToolsFramework/Private/ScriptableInteractiveTool.cpp
+++ b/Engine/Plugins/Runtime/ScriptableToolsFramework/Source/ScriptableToolsFramework/Private/ScriptableInteractiveTool.cpp
@@ -38,6 +38,88 @@ namespace ScriptableInteractiveToolLocals
 	class SScriptableToolViewportOverlayDragContainer : public SObjectWidget
 	{ };
 }
+namespace UELocal
+{
+
+	template<typename PropertyType>
+	PropertyType* FindValidPropertyByTypeAndName(
+		UScriptableInteractiveToolPropertySet* PropertySet,
+		FString PropertyName,
+		FString CallerString)
+	{
+		if (PropertySet == nullptr)
+		{
+			UE_LOG(LogTemp, Warning, TEXT("%s: Property Set is null"), *CallerString);
+			return nullptr;
+		}
+
+		for (TFieldIterator<FProperty> PropIt(PropertySet->GetClass(), EFieldIterationFlags::None); PropIt; ++PropIt)
+		{
+			FProperty* Property = *PropIt;
+			if (Property->GetName() == PropertyName)
+			{
+				if (PropertyType* TypedProp = CastField<PropertyType>(*PropIt))
+				{
+					return TypedProp;
+				}
+				else
+				{
+					UE_LOG(LogTemp, Warning, TEXT("%s: Property %s does not have a compatible Type"), *CallerString, *PropertyName);
+					return nullptr;
+				}
+			}
+		}
+
+		UE_LOG(LogTemp, Warning, TEXT("%s: Could not find a Property named %s in the Property Set %s"), *CallerString, *PropertyName,
+			*PropertySet->GetName());
+		return nullptr;
+	}
+
+} // end namespace UELocal
+
+
+UScriptableInteractiveToolPropertySet* UScriptableInteractiveToolPropertySet::SetPropertyAsReadOnly(FString PropertyName, bool bReadOnly)
+{
+	FProperty* Property = UELocal::FindValidPropertyByTypeAndName<FProperty>(
+		this, PropertyName, TEXT("SetReadOnly"));
+	if (Property == nullptr)
+	{
+		return this;
+	}
+	if (!bReadOnly)
+	{
+		Property->ClearPropertyFlags(EPropertyFlags::CPF_EditConst);
+		Property->SetPropertyFlags(EPropertyFlags::CPF_Edit);
+	}
+	else
+	{
+		Property->ClearPropertyFlags(EPropertyFlags::CPF_Edit);
+		Property->SetPropertyFlags(EPropertyFlags::CPF_EditConst);
+	}
+
+	return this;
+}
+
+UScriptableInteractiveToolPropertySet* UScriptableInteractiveToolPropertySet::SetPropertyAsHidden(FString PropertyName, bool bVisible)
+{
+	FProperty* Property = UELocal::FindValidPropertyByTypeAndName<FProperty>(
+		this, PropertyName, TEXT("SetHidden"));
+	if (Property == nullptr)
+	{
+		return this;
+	}
+	if (!bVisible)
+	{
+		Property->ClearPropertyFlags(EPropertyFlags::CPF_EditConst);
+		Property->ClearPropertyFlags(EPropertyFlags::CPF_Edit);
+	}
+	else
+	{
+		Property->SetPropertyFlags(EPropertyFlags::CPF_Edit);
+	}
+
+	return this;
+}
 
 UScriptableInteractiveTool* UScriptableInteractiveToolPropertySet::GetOwningTool(EToolsFrameworkOutcomePins& Outcome)
 {
@@ -520,44 +602,6 @@ UScriptableInteractiveToolPropertySet* UScriptableInteractiveTool::SavePropertyS
 
 
 
-namespace UELocal
-{
-
-template<typename PropertyType>
-PropertyType* FindValidPropertyByTypeAndName(
-	UScriptableInteractiveToolPropertySet* PropertySet,
-	FString PropertyName,
-	FString CallerString )
-{
-	if ( PropertySet == nullptr )
-	{
-		UE_LOG(LogTemp, Warning, TEXT("%s: Property Set is null"), *CallerString);
-		return nullptr;
-	}
-
-	for (TFieldIterator<FProperty> PropIt(PropertySet->GetClass(), EFieldIterationFlags::None); PropIt; ++PropIt)
-	{
-		FProperty* Property = *PropIt;
-		if (Property->GetName() == PropertyName)
-		{
-			if (PropertyType* TypedProp = CastField<PropertyType>(*PropIt))
-			{
-				return TypedProp;
-			}
-			else
-			{
-				UE_LOG(LogTemp, Warning, TEXT("%s: Property %s does not have a compatible Type"), *CallerString, *PropertyName);
-				return nullptr;
-			}
-		}
-	}
-
-	UE_LOG(LogTemp, Warning, TEXT("%s: Could not find a Property named %s in the Property Set %s"), *CallerString, *PropertyName, 
-		*PropertySet->GetName());
-	return nullptr;
-}
-
-} // end namespace UELocal
 
 
 UScriptableInteractiveToolPropertySet* UScriptableInteractiveTool::WatchFloatProperty(
diff --git a/Engine/Plugins/Runtime/ScriptableToolsFramework/Source/ScriptableToolsFramework/Public/ScriptableInteractiveTool.h b/Engine/Plugins/Runtime/ScriptableToolsFramework/Source/ScriptableToolsFramework/Public/ScriptableInteractiveTool.h
--- a/Engine/Plugins/Runtime/ScriptableToolsFramework/Source/ScriptableToolsFramework/Public/ScriptableInteractiveTool.h
+++ b/Engine/Plugins/Runtime/ScriptableToolsFramework/Source/ScriptableToolsFramework/Public/ScriptableInteractiveTool.h
@@ -277,6 +277,14 @@ class SCRIPTABLETOOLSFRAMEWORK_API UScriptableInteractiveToolPropertySet : publi
 	UScriptableInteractiveTool* GetOwningTool(
 		EToolsFrameworkOutcomePins& Outcome);
 
+	/** Toggle whether or not the named property should be displayed as a readonly field in the tool's details view. */
+	UFUNCTION(BlueprintCallable, Category = "ScriptableTool")
+	UPARAM(DisplayName = "Property Set") UScriptableInteractiveToolPropertySet* SetPropertyAsReadOnly(FString PropertyName, bool bReadOnly);
+
+	/** Toggle whether or not the named property should be displayed at all as a field in the tool's details view. */
+	UFUNCTION(BlueprintCallable, Category = "ScriptableTool")
+	UPARAM(DisplayName = "Property Set") UScriptableInteractiveToolPropertySet* SetPropertyAsHidden(FString PropertyName, bool bVisible);
+
 	virtual UWorld* GetWorld() const override;	// must implement this for WorldContextObject dependent BP functions to show.
 
 protected:
