ScriptableTools: Rename SetPropertyAsHidden parameter to bHidden.

diff --git a/Engine/Plugins/Runtime/ScriptableToolsFramework/Source/ScriptableToolsFramework/Private/ScriptableInteractiveTool.cpp b/Engine/Plugins/Runtime/ScriptableToolsFramework/Source/ScriptableToolsFramework/Private/ScriptableInteractiveTool.cpp
--- a/Engine/Plugins/Runtime/ScriptableToolsFramework/Source/ScriptableToolsFramework/Private/ScriptableInteractiveTool.cpp
+++ b/Engine/Plugins/Runtime/ScriptableToolsFramework/Source/ScriptableToolsFramework/Private/ScriptableInteractiveTool.cpp
@@ -100,7 +100,7 @@ UScriptableInteractiveToolPropertySet* UScriptableInteractiveToolPropertySet::Se
 	return this;
 }
 
-UScriptableInteractiveToolPropertySet* UScriptableInteractiveToolPropertySet::SetPropertyAsHidden(FString PropertyName, bool bVisible)
+UScriptableInteractiveToolPropertySet* UScriptableInteractiveToolPropertySet::SetPropertyAsHidden(FString PropertyName, bool bHidden)
 {
 	FProperty* Property = UELocal::FindValidPropertyByTypeAndName<FProperty>(
 		this, PropertyName, TEXT("SetHidden"));
@@ -108,7 +108,7 @@ UScriptableInteractiveToolPropertySet* UScriptableInteractiveToolPropertySet::Se
 	{
 		return this;
 	}
-	if (!bVisible)
+	if (bHidden)
 	{
 		Property->ClearPropertyFlags(EPropertyFlags::CPF_EditConst);
 		Property->ClearPropertyFlags(EPropertyFlags::CPF_Edit);
diff --git a/Engine/Plugins/Runtime/ScriptableToolsFramework/Source/ScriptableToolsFramework/Public/ScriptableInteractiveTool.h b/Engine/Plugins/Runtime/ScriptableToolsFramework/Source/ScriptableToolsFramework/Public/ScriptableInteractiveTool.h
--- a/Engine/Plugins/Runtime/ScriptableToolsFramework/Source/ScriptableToolsFramework/Public/ScriptableInteractiveTool.h
+++ b/Engine/Plugins/Runtime/ScriptableToolsFramework/Source/ScriptableToolsFramework/Public/ScriptableInteractiveTool.h
@@ -283,7 +283,7 @@ class SCRIPTABLETOOLSFRAMEWORK_API UScriptableInteractiveToolPropertySet : publi
 
 	/** Toggle whether or not the named property should be displayed at all as a field in the tool's details view. */
 	UFUNCTION(BlueprintCallable, Category = "ScriptableTool")
-	UPARAM(DisplayName = "Property Set") UScriptableInteractiveToolPropertySet* SetPropertyAsHidden(FString PropertyName, bool bVisible);
+	UPARAM(DisplayName = "Property Set") UScriptableInteractiveToolPropertySet* SetPropertyAsHidden(FString PropertyName, bool bHidden);
 
 	virtual UWorld* GetWorld() const override;	// must implement this for WorldContextObject dependent BP functions to show.
 
