Updated headers using UnrealCodeFixup to make sure dllstorage is on methods/staticvars instead of types. Part 4/n

diff --git a/Engine/Plugins/Runtime/ScriptableToolsFramework/Source/EditorScriptableToolsFramework/Public/EditorScriptableInteractiveTool.h b/Engine/Plugins/Runtime/ScriptableToolsFramework/Source/EditorScriptableToolsFramework/Public/EditorScriptableInteractiveTool.h
--- a/Engine/Plugins/Runtime/ScriptableToolsFramework/Source/EditorScriptableToolsFramework/Public/EditorScriptableInteractiveTool.h
+++ b/Engine/Plugins/Runtime/ScriptableToolsFramework/Source/EditorScriptableToolsFramework/Public/EditorScriptableInteractiveTool.h
@@ -8,8 +8,8 @@
 /**
  * Editor-Only variant of UScriptableInteractiveTool, which gives access to Editor-Only BP functions
  */
-UCLASS(Transient, Blueprintable)
-class EDITORSCRIPTABLETOOLSFRAMEWORK_API UEditorScriptableInteractiveTool : public UScriptableInteractiveTool
+UCLASS(MinimalAPI, Transient, Blueprintable)
+class UEditorScriptableInteractiveTool : public UScriptableInteractiveTool
 {
 	GENERATED_BODY()
 public:
@@ -17,8 +17,8 @@ public:
 };
 
 
-UCLASS(Transient, Blueprintable)
-class EDITORSCRIPTABLETOOLSFRAMEWORK_API UEditorScriptableInteractiveToolPropertySet : public UScriptableInteractiveToolPropertySet
+UCLASS(MinimalAPI, Transient, Blueprintable)
+class UEditorScriptableInteractiveToolPropertySet : public UScriptableInteractiveToolPropertySet
 {
 	GENERATED_BODY()
 public:
diff --git a/Engine/Plugins/Runtime/ScriptableToolsFramework/Source/EditorScriptableToolsFramework/Public/ScriptableToolSet.h b/Engine/Plugins/Runtime/ScriptableToolsFramework/Source/EditorScriptableToolsFramework/Public/ScriptableToolSet.h
--- a/Engine/Plugins/Runtime/ScriptableToolsFramework/Source/EditorScriptableToolsFramework/Public/ScriptableToolSet.h
+++ b/Engine/Plugins/Runtime/ScriptableToolsFramework/Source/EditorScriptableToolsFramework/Public/ScriptableToolSet.h
@@ -5,6 +5,8 @@
 #include "CoreMinimal.h"
 #include "ScriptableToolSet.generated.h"
 
+#define UE_API EDITORSCRIPTABLETOOLSFRAMEWORK_API
+
 class UScriptableInteractiveTool;
 class UBaseScriptableToolBuilder;
 struct FScriptableToolGroupSet;
@@ -17,39 +19,39 @@ DECLARE_DELEGATE_OneParam(FToolsLoadingUpdateDelegate, TSharedRef<struct FStream
 /**
  * UScriptableToolSet represents a set of UScriptableInteractiveTool types.
  */
-UCLASS()
-class EDITORSCRIPTABLETOOLSFRAMEWORK_API UScriptableToolSet : public UObject
+UCLASS(MinimalAPI)
+class UScriptableToolSet : public UObject
 {
 	GENERATED_BODY()
 
 public:
 
-	UScriptableToolSet();
+	UE_API UScriptableToolSet();
 
-	virtual ~UScriptableToolSet();
+	UE_API virtual ~UScriptableToolSet();
 
 	/**
 	* Forces the unloading of all tools loaded
 	*/
-	void UnloadAllTools();
+	UE_API void UnloadAllTools();
 
 	/**
 	 * Find all UScriptableInteractiveTool classes in the current project.
 	 * (Currently no support for filtering/etc)
 	 */
-	void ReinitializeScriptableTools(FPreToolsLoadedDelegate PreDelegate, FToolsLoadedDelegate PostDelegate, FToolsLoadingUpdateDelegate UpdateDelegate, FScriptableToolGroupSet* TagsToFilter = nullptr);
+	UE_API void ReinitializeScriptableTools(FPreToolsLoadedDelegate PreDelegate, FToolsLoadedDelegate PostDelegate, FToolsLoadingUpdateDelegate UpdateDelegate, FScriptableToolGroupSet* TagsToFilter = nullptr);
 
 	/**
 	 * Allow external code to process each UScriptableInteractiveTool in the current ToolSet
 	 */
-	void ForEachScriptableTool(
+	UE_API void ForEachScriptableTool(
 		TFunctionRef<void(UClass* ToolClass, UBaseScriptableToolBuilder* ToolBuilder)> ProcessToolFunc);
 
 private:
 
-	void HandleAssetCanDelete(const TArray<UObject*>& InObjectsToDelete, FCanDeleteAssetResult& OutCanDelete);
+	UE_API void HandleAssetCanDelete(const TArray<UObject*>& InObjectsToDelete, FCanDeleteAssetResult& OutCanDelete);
 
-	void PostToolLoad(FToolsLoadedDelegate Delegate, TArray< FSoftObjectPath > ObjectsLoaded, TSharedPtr<FScriptableToolGroupSet> TagsToFilter);
+	UE_API void PostToolLoad(FToolsLoadedDelegate Delegate, TArray< FSoftObjectPath > ObjectsLoaded, TSharedPtr<FScriptableToolGroupSet> TagsToFilter);
 
 	bool bActiveLoading = false;
 	TSharedPtr<FStreamableHandle> AsyncLoadHandle;
@@ -69,3 +71,5 @@ private:
 	UPROPERTY()
 	TArray<TObjectPtr<UBaseScriptableToolBuilder>> ToolBuilders;
 };
+
+#undef UE_API
diff --git a/Engine/Plugins/Runtime/ScriptableToolsFramework/Source/ScriptableToolsFramework/Public/BaseTools/ScriptableClickDragTool.h b/Engine/Plugins/Runtime/ScriptableToolsFramework/Source/ScriptableToolsFramework/Public/BaseTools/ScriptableClickDragTool.h
--- a/Engine/Plugins/Runtime/ScriptableToolsFramework/Source/ScriptableToolsFramework/Public/BaseTools/ScriptableClickDragTool.h
+++ b/Engine/Plugins/Runtime/ScriptableToolsFramework/Source/ScriptableToolsFramework/Public/BaseTools/ScriptableClickDragTool.h
@@ -8,6 +8,8 @@
 
 #include "ScriptableClickDragTool.generated.h"
 
+#define UE_API SCRIPTABLETOOLSFRAMEWORK_API
+
 class UClickDragInputBehavior;
 class UMouseHoverBehavior;
 
@@ -19,8 +21,8 @@ class UMouseHoverBehavior;
  * Note: Currently these interactions are only supported for mouse devices, or devices remapped to 
  * be a mouse by the UInputRouter. Also currently only the left-mouse-button is supported.
  */
-UCLASS()
-class SCRIPTABLETOOLSFRAMEWORK_API UScriptableClickDragTool : public UScriptableInteractiveTool, public IClickDragBehaviorTarget, public IHoverBehaviorTarget
+UCLASS(MinimalAPI)
+class UScriptableClickDragTool : public UScriptableInteractiveTool, public IClickDragBehaviorTarget, public IHoverBehaviorTarget
 {
 	GENERATED_BODY()
 
@@ -45,36 +47,36 @@ public:
 public:
 
 	UFUNCTION(BlueprintNativeEvent, Category = "ClickDrag Events")
-	FInputRayHit TestIfCanBeginClickDrag(FInputDeviceRay ClickPos, const FScriptableToolModifierStates& Modifiers);
+	UE_API FInputRayHit TestIfCanBeginClickDrag(FInputDeviceRay ClickPos, const FScriptableToolModifierStates& Modifiers);
 
-	virtual FInputRayHit TestIfCanBeginClickDrag_Implementation(FInputDeviceRay ClickPos, const FScriptableToolModifierStates& Modifiers);
+	UE_API virtual FInputRayHit TestIfCanBeginClickDrag_Implementation(FInputDeviceRay ClickPos, const FScriptableToolModifierStates& Modifiers);
 
 
 	UFUNCTION(BlueprintNativeEvent, Category = "ClickDrag Events")
-	void OnDragBegin(FInputDeviceRay StartPosition, const FScriptableToolModifierStates& Modifiers);
+	UE_API void OnDragBegin(FInputDeviceRay StartPosition, const FScriptableToolModifierStates& Modifiers);
 
-	virtual void OnDragBegin_Implementation(FInputDeviceRay StartPosition, const FScriptableToolModifierStates& Modifiers);
+	UE_API virtual void OnDragBegin_Implementation(FInputDeviceRay StartPosition, const FScriptableToolModifierStates& Modifiers);
 
 
 	UFUNCTION(BlueprintNativeEvent, Category = "ClickDrag Events")
-	void OnDragUpdatePosition(FInputDeviceRay NewPosition, const FScriptableToolModifierStates& Modifiers);
+	UE_API void OnDragUpdatePosition(FInputDeviceRay NewPosition, const FScriptableToolModifierStates& Modifiers);
 
-	virtual void OnDragUpdatePosition_Implementation(FInputDeviceRay NewPosition, const FScriptableToolModifierStates& Modifiers);
+	UE_API virtual void OnDragUpdatePosition_Implementation(FInputDeviceRay NewPosition, const FScriptableToolModifierStates& Modifiers);
 
 
 	UFUNCTION(BlueprintNativeEvent, Category = "ClickDrag Events")
-	void OnDragEnd(FInputDeviceRay EndPosition, const FScriptableToolModifierStates& Modifiers);
+	UE_API void OnDragEnd(FInputDeviceRay EndPosition, const FScriptableToolModifierStates& Modifiers);
 
-	virtual void OnDragEnd_Implementation(FInputDeviceRay EndPosition, const FScriptableToolModifierStates& Modifiers);
+	UE_API virtual void OnDragEnd_Implementation(FInputDeviceRay EndPosition, const FScriptableToolModifierStates& Modifiers);
 
 
 	UFUNCTION(BlueprintNativeEvent, Category = "ClickDrag Events")
-	void OnDragSequenceCancelled();
+	UE_API void OnDragSequenceCancelled();
 
-	virtual void OnDragSequenceCancelled_Implementation();
+	UE_API virtual void OnDragSequenceCancelled_Implementation();
 
 	UFUNCTION(BlueprintCallable, Category = "ScriptableTool|Input")
-	UPARAM(DisplayName="In ClickDrag State") bool InActiveClickDrag() const;
+	UE_API UPARAM(DisplayName="In ClickDrag State") bool InActiveClickDrag() const;
 
 	//
 	// Hover Support API
@@ -95,9 +97,9 @@ public:
 	 * @return a FInputRayHit indicating a hit (true/false) and the hit depth along the hit-ray (0 for "always consume")
 	 */
 	UFUNCTION(BlueprintNativeEvent, Category = "Hover Events")
-	FInputRayHit OnHoverHitTest(FInputDeviceRay HoverPos, const FScriptableToolModifierStates& Modifiers);
+	UE_API FInputRayHit OnHoverHitTest(FInputDeviceRay HoverPos, const FScriptableToolModifierStates& Modifiers);
 
-	virtual FInputRayHit OnHoverHitTest_Implementation(FInputDeviceRay HoverPos, const FScriptableToolModifierStates& Modifiers);
+	UE_API virtual FInputRayHit OnHoverHitTest_Implementation(FInputDeviceRay HoverPos, const FScriptableToolModifierStates& Modifiers);
 
 	/**
 	 * Begin a sequence of Hover input events 
@@ -105,9 +107,9 @@ public:
 	 * @param Modifiers current modifier key/button state
 	 */
 	UFUNCTION(BlueprintNativeEvent, Category = "Hover Events")
-	void OnHoverBegin(FInputDeviceRay HoverPos, const FScriptableToolModifierStates& Modifiers);
+	UE_API void OnHoverBegin(FInputDeviceRay HoverPos, const FScriptableToolModifierStates& Modifiers);
 
-	virtual void OnHoverBegin_Implementation(FInputDeviceRay HoverPos, const FScriptableToolModifierStates& Modifiers);
+	UE_API virtual void OnHoverBegin_Implementation(FInputDeviceRay HoverPos, const FScriptableToolModifierStates& Modifiers);
 
 	/**
 	 * Update an active Hover sequence. OnHoverUpdate is only ever called between OnHoverBegin and OnHoverEnd
@@ -116,22 +118,22 @@ public:
 	 * @return true to continue hovering, false to stop receiving additional hover events
 	 */
 	UFUNCTION(BlueprintNativeEvent, Category = "Hover Events")
-	UPARAM(DisplayName="Continue Hovering") bool
+	UE_API UPARAM(DisplayName="Continue Hovering") bool
 	OnHoverUpdate(FInputDeviceRay HoverPos, const FScriptableToolModifierStates& Modifiers);
 
-	virtual bool OnHoverUpdate_Implementation(FInputDeviceRay HoverPos, const FScriptableToolModifierStates& Modifiers);
+	UE_API virtual bool OnHoverUpdate_Implementation(FInputDeviceRay HoverPos, const FScriptableToolModifierStates& Modifiers);
 
 	/**
 	 * Indicate that an active captured Hover sequence has ended. The device may no longer be in the viewport, so no position can be provided.
 	 * @param Modifiers current modifier key/button state
 	 */
 	UFUNCTION(BlueprintNativeEvent, Category = "Hover Events")
-	void OnHoverEnd(const FScriptableToolModifierStates& Modifiers);
+	UE_API void OnHoverEnd(const FScriptableToolModifierStates& Modifiers);
 
-	virtual void OnHoverEnd_Implementation(const FScriptableToolModifierStates& Modifiers);
+	UE_API virtual void OnHoverEnd_Implementation(const FScriptableToolModifierStates& Modifiers);
 
 	UFUNCTION(BlueprintCallable, Category = "ScriptableTool|Input")
-	UPARAM(DisplayName="In Hover State") bool InActiveHover() const;
+	UE_API UPARAM(DisplayName="In Hover State") bool InActiveHover() const;
 
 	//
 	// Modifer Buttons Support API
@@ -140,24 +142,24 @@ public:
 
 	/** @return true if the Shift key is currently held down */
 	UFUNCTION(BlueprintCallable, Category="ScriptableTool|Input")
-	UPARAM(DisplayName="Shift Down") bool IsShiftDown() const;
+	UE_API UPARAM(DisplayName="Shift Down") bool IsShiftDown() const;
 
 	/** @return true if the Ctrl key is currently held down */
 	UFUNCTION(BlueprintCallable, Category="ScriptableTool|Input")
-	UPARAM(DisplayName="Ctrl Down") bool IsCtrlDown() const;
+	UE_API UPARAM(DisplayName="Ctrl Down") bool IsCtrlDown() const;
 
 	/** @return true if the Alt key is currently held down */
 	UFUNCTION(BlueprintCallable, Category="ScriptableTool|Input")
-	UPARAM(DisplayName="Alt Down") bool IsAltDown() const;
+	UE_API UPARAM(DisplayName="Alt Down") bool IsAltDown() const;
 
 	/** @return a struct containing the current Modifier key states */
 	UFUNCTION(BlueprintCallable, Category="ScriptableTool|Input")
-	FScriptableToolModifierStates GetActiveModifiers();
+	UE_API FScriptableToolModifierStates GetActiveModifiers();
 
 
 
 public:
-	virtual void Setup() override;
+	UE_API virtual void Setup() override;
 
 	UPROPERTY(Transient, DuplicateTransient, NonTransactional, SkipSerialization)
 	TObjectPtr<UClickDragInputBehavior> ClickDragBehavior;
@@ -174,17 +176,19 @@ public:
 	bool bAltModifier = false;
 
 	// IClickDragBehaviorTarget API
-	virtual FInputRayHit CanBeginClickDragSequence(const FInputDeviceRay& PressPos) override;
-	virtual void OnClickPress(const FInputDeviceRay& PressPos) override;
-	virtual void OnClickDrag(const FInputDeviceRay& DragPos) override;
-	virtual void OnClickRelease(const FInputDeviceRay& ReleasePos) override;
-	virtual void OnTerminateDragSequence() override;
-	virtual void OnUpdateModifierState(int ModifierID, bool bIsOn);
+	UE_API virtual FInputRayHit CanBeginClickDragSequence(const FInputDeviceRay& PressPos) override;
+	UE_API virtual void OnClickPress(const FInputDeviceRay& PressPos) override;
+	UE_API virtual void OnClickDrag(const FInputDeviceRay& DragPos) override;
+	UE_API virtual void OnClickRelease(const FInputDeviceRay& ReleasePos) override;
+	UE_API virtual void OnTerminateDragSequence() override;
+	UE_API virtual void OnUpdateModifierState(int ModifierID, bool bIsOn);
 
 	// IHoverBehaviorTarget API
-	virtual FInputRayHit BeginHoverSequenceHitTest(const FInputDeviceRay& PressPos);
-	virtual void OnBeginHover(const FInputDeviceRay& DevicePos);
-	virtual bool OnUpdateHover(const FInputDeviceRay& DevicePos);
-	virtual void OnEndHover();
+	UE_API virtual FInputRayHit BeginHoverSequenceHitTest(const FInputDeviceRay& PressPos);
+	UE_API virtual void OnBeginHover(const FInputDeviceRay& DevicePos);
+	UE_API virtual bool OnUpdateHover(const FInputDeviceRay& DevicePos);
+	UE_API virtual void OnEndHover();
+
+};
 
-};
\ No newline at end of file
+#undef UE_API
diff --git a/Engine/Plugins/Runtime/ScriptableToolsFramework/Source/ScriptableToolsFramework/Public/BaseTools/ScriptableModularBehaviorTool.h b/Engine/Plugins/Runtime/ScriptableToolsFramework/Source/ScriptableToolsFramework/Public/BaseTools/ScriptableModularBehaviorTool.h
--- a/Engine/Plugins/Runtime/ScriptableToolsFramework/Source/ScriptableToolsFramework/Public/BaseTools/ScriptableModularBehaviorTool.h
+++ b/Engine/Plugins/Runtime/ScriptableToolsFramework/Source/ScriptableToolsFramework/Public/BaseTools/ScriptableModularBehaviorTool.h
@@ -8,6 +8,8 @@
 #include "Behaviors/ScriptableToolBehaviorDelegates.h"
 #include "ScriptableModularBehaviorTool.generated.h"
 
+#define UE_API SCRIPTABLETOOLSFRAMEWORK_API
+
 class UScriptableToolSingleClickBehavior;
 class UScriptableToolDoubleClickBehavior;
 class UScriptableToolClickDragBehavior;
@@ -20,8 +22,8 @@ class UScriptableToolKeyInputBehavior;
 /**
  A Scriptable tool base blueprint class which provides support for user defined mouse interaction behaviors
  */ 
-UCLASS()
-class SCRIPTABLETOOLSFRAMEWORK_API UScriptableModularBehaviorTool :	public UScriptableInteractiveTool
+UCLASS(MinimalAPI)
+class UScriptableModularBehaviorTool :	public UScriptableInteractiveTool
 {
 	GENERATED_BODY()
 
@@ -47,7 +49,7 @@ public:
 	*	@param MouseButton Determines which mouse button the behavior captures on
 	*/
 	UFUNCTION(BlueprintCallable, Category = "ScriptableTool|Input", meta = (AdvancedDisplay = "CaptureCheck, CapturePriority, MouseButton, bHitTestOnRelease"))
-	void AddSingleClickBehavior(
+	UE_API void AddSingleClickBehavior(
 		const FTestIfHitByClickDelegate IfHitByClick,
 		const FOnHitByClickDelegate OnHitByClick,
 		const FMouseBehaviorModiferCheckDelegate CaptureCheck,
@@ -73,7 +75,7 @@ public:
 	*	@param MouseButton Determines which mouse button the behavior captures on
 	*/
 	UFUNCTION(BlueprintCallable, Category = "ScriptableTool|Input", meta = (AdvancedDisplay = "CaptureCheck, CapturePriority, MouseButton, bHitTestOnRelease"))
-	void AddDoubleClickBehavior(
+	UE_API void AddDoubleClickBehavior(
 		const FTestIfHitByClickDelegate IfHitByClick,
 		const FOnHitByClickDelegate OnHitByClick,
 		const FMouseBehaviorModiferCheckDelegate CaptureCheck,
@@ -102,7 +104,7 @@ public:
 	*	@param MouseButton Determines which mouse button the behavior captures on
 	*/
 	UFUNCTION(BlueprintCallable, Category = "ScriptableTool|Input", meta = (AdvancedDisplay = "CaptureCheck, CapturePriority, MouseButton, bUpdateModifiersDuringDrag"))
-	void AddClickDragBehavior(		
+	UE_API void AddClickDragBehavior(		
 		const FTestCanBeginClickDragSequenceDelegate CanBeginClickDragSequence,
 		const FOnClickPressDelegate OnClickPress,
 		const FOnClickDragDelegate OnClickDrag,
@@ -140,7 +142,7 @@ public:
 	*	@param MouseButton Determines which mouse button the behavior captures on
 	*/
 	UFUNCTION(BlueprintCallable, Category = "ScriptableTool|Input", meta = (AdvancedDisplay = "CaptureCheck, CapturePriority, MouseButton, bBeginDragIfClickTargetNotHit, ClickDistanceThreshold"))
-	void AddSingleClickOrDragBehavior(		
+	UE_API void AddSingleClickOrDragBehavior(		
 		FTestIfHitByClickDelegate IfHitByClick,
 		FOnHitByClickDelegate OnHitByClick,
 		FTestCanBeginClickDragSequenceDelegate CanBeginClickDragSequence,
@@ -163,7 +165,7 @@ public:
 	*	@param CapturePriority The priority is used to resolve situations where multiple behaviors want the same capture
 	*/
 	UFUNCTION(BlueprintCallable, Category = "ScriptableTool|Input", meta = (AdvancedDisplay = "CaptureCheck, CapturePriority"))
-	void AddMouseWheelBehavior(		
+	UE_API void AddMouseWheelBehavior(		
 		FTestShouldRespondToMouseWheelDelegate TestShouldRespondToMouseWheel,
 		FOnMouseWheelScrollUpDelegate OnMouseWheelScrollUp,
 		FOnMouseWheelScrollDownDelegate OnMouseWheelScrollDown,
@@ -198,7 +200,7 @@ public:
 	*	@param MouseButton Determines which mouse button the behavior captures on
 	*/
 	UFUNCTION(BlueprintCallable, Category = "ScriptableTool|Input", meta = (AdvancedDisplay = "CaptureCheck, HoverCaptureCheck, CapturePriority, MouseButton"))
-	void AddMultiClickSequenceBehavior(		
+	UE_API void AddMultiClickSequenceBehavior(		
 		FOnBeginSequencePreviewDelegate OnBeginSequencePreview,
 		FCanBeginClickSequenceDelegate CanBeginClickSequence,
 		FOnBeginClickSequenceDelegate OnBeginClickSequence,
@@ -221,7 +223,7 @@ public:
 	*	@param CapturePriority The priority is used to resolve situations where multiple behaviors want the same capture
 	*/
 	UFUNCTION(BlueprintCallable, Category = "ScriptableTool|Input", meta = (AdvancedDisplay = "CaptureCheck, CapturePriority"))
-	void AddMouseHoverBehavior(
+	UE_API void AddMouseHoverBehavior(
 		FBeginHoverSequenceHitTestDelegate BeginHoverSequenceHitTest,
 		FOnBeginHoverDelegate OnBeginHover,
 		FOnUpdateHoverDelegate OnUpdateHover,
@@ -244,7 +246,7 @@ public:
 
 
 	UFUNCTION(BlueprintCallable, Category = "ScriptableTool|Input", meta = (AdvancedDisplay = "CaptureCheck, CapturePriority"))
-	void AddSingleKeyInputBehavior(
+	UE_API void AddSingleKeyInputBehavior(
 		FOnKeyStateToggleDelegate OnKeyPressed,
 		FOnKeyStateToggleDelegate OnKeyReleased,
 		FOnForceEndCaptureDelegate_ScriptableTools OnForceEndCaptureFuncIn,
@@ -279,7 +281,7 @@ public:
 
 
 	UFUNCTION(BlueprintCallable, Category = "ScriptableTool|Input", meta = (AdvancedDisplay = "CaptureCheck, CapturePriority"))
-		void AddMultiKeyInputBehavior(
+		UE_API void AddMultiKeyInputBehavior(
 			FOnKeyStateToggleDelegate OnKeyPressed,
 			FOnKeyStateToggleDelegate OnKeyReleased,
 			FOnForceEndCaptureDelegate_ScriptableTools OnForceEndCaptureFuncIn,
@@ -337,21 +339,21 @@ public:
 
 	/** @return true if the Shift key is currently held down */
 	UFUNCTION(BlueprintCallable, Category="ScriptableTool|Input")
-	UPARAM(DisplayName="Shift Down") bool IsShiftDown() const;
+	UE_API UPARAM(DisplayName="Shift Down") bool IsShiftDown() const;
 
 	/** @return true if the Ctrl key is currently held down */
 	UFUNCTION(BlueprintCallable, Category="ScriptableTool|Input")
-	UPARAM(DisplayName="Ctrl Down") bool IsCtrlDown() const;
+	UE_API UPARAM(DisplayName="Ctrl Down") bool IsCtrlDown() const;
 
 	/** @return true if the Alt key is currently held down */
 	UFUNCTION(BlueprintCallable, Category="ScriptableTool|Input")
-	UPARAM(DisplayName="Alt Down") bool IsAltDown() const;
+	UE_API UPARAM(DisplayName="Alt Down") bool IsAltDown() const;
 
 	/** @return a struct containing the current Modifier key states */
 	UFUNCTION(BlueprintCallable, Category="ScriptableTool|Input")
-	FScriptableToolModifierStates GetActiveModifiers();
+	UE_API FScriptableToolModifierStates GetActiveModifiers();
 
-	void OnUpdateModifierState(int ModifierID, bool bIsOn);
+	UE_API void OnUpdateModifierState(int ModifierID, bool bIsOn);
 
 private:
 
@@ -380,3 +382,4 @@ private:
 
 
 
+#undef UE_API
diff --git a/Engine/Plugins/Runtime/ScriptableToolsFramework/Source/ScriptableToolsFramework/Public/BaseTools/ScriptableSingleClickTool.h b/Engine/Plugins/Runtime/ScriptableToolsFramework/Source/ScriptableToolsFramework/Public/BaseTools/ScriptableSingleClickTool.h
--- a/Engine/Plugins/Runtime/ScriptableToolsFramework/Source/ScriptableToolsFramework/Public/BaseTools/ScriptableSingleClickTool.h
+++ b/Engine/Plugins/Runtime/ScriptableToolsFramework/Source/ScriptableToolsFramework/Public/BaseTools/ScriptableSingleClickTool.h
@@ -8,6 +8,8 @@
 
 #include "ScriptableSingleClickTool.generated.h"
 
+#define UE_API SCRIPTABLETOOLSFRAMEWORK_API
+
 class USingleClickInputBehavior;
 class UMouseHoverBehavior;
 
@@ -24,8 +26,8 @@ class UMouseHoverBehavior;
  * Note: Currently these interactions are only supported for mouse devices, or devices remapped to 
  * be a mouse by the UInputRouter. Also currently only the left-mouse-button is supported.
  */
-UCLASS(Hidden)
-class SCRIPTABLETOOLSFRAMEWORK_API UScriptableSingleClickTool : public UScriptableInteractiveTool, public IClickBehaviorTarget, public IHoverBehaviorTarget
+UCLASS(MinimalAPI, Hidden)
+class UScriptableSingleClickTool : public UScriptableInteractiveTool, public IClickBehaviorTarget, public IHoverBehaviorTarget
 {
 	GENERATED_BODY()
 
@@ -54,9 +56,9 @@ public:
 	 * @return a FInputRayHit indicating a hit (true/false) and the hit depth along the hit-ray (0 for "always consume")
 	 */
 	UFUNCTION(BlueprintNativeEvent, Category = "Click Events")
-	FInputRayHit TestIfHitByClick(FInputDeviceRay ClickPos, const FScriptableToolModifierStates& Modifiers);
+	UE_API FInputRayHit TestIfHitByClick(FInputDeviceRay ClickPos, const FScriptableToolModifierStates& Modifiers);
 
-	virtual FInputRayHit TestIfHitByClick_Implementation(FInputDeviceRay ClickPos, const FScriptableToolModifierStates& Modifiers);
+	UE_API virtual FInputRayHit TestIfHitByClick_Implementation(FInputDeviceRay ClickPos, const FScriptableToolModifierStates& Modifiers);
 
 	/**
 	 * OnHitByClick is called when the Tool has indicated it would like to consume a click event (via TestIfHitByClick), and the click was not cancelled
@@ -65,9 +67,9 @@ public:
 	 * @param Modifiers current modifier key/button state
 	 */
 	UFUNCTION(BlueprintNativeEvent, Category = "Click Events")
-	void OnHitByClick(FInputDeviceRay ClickPos, const FScriptableToolModifierStates& Modifiers);
+	UE_API void OnHitByClick(FInputDeviceRay ClickPos, const FScriptableToolModifierStates& Modifiers);
 
-	virtual void OnHitByClick_Implementation(FInputDeviceRay ClickPos, const FScriptableToolModifierStates& Modifiers);
+	UE_API virtual void OnHitByClick_Implementation(FInputDeviceRay ClickPos, const FScriptableToolModifierStates& Modifiers);
 
 
 	//
@@ -89,9 +91,9 @@ public:
 	 * @return a FInputRayHit indicating a hit (true/false) and the hit depth along the hit-ray (0 for "always consume")
 	 */
 	UFUNCTION(BlueprintNativeEvent, Category = "Hover Events")
-	FInputRayHit OnHoverHitTest(FInputDeviceRay HoverPos, const FScriptableToolModifierStates& Modifiers);
+	UE_API FInputRayHit OnHoverHitTest(FInputDeviceRay HoverPos, const FScriptableToolModifierStates& Modifiers);
 
-	virtual FInputRayHit OnHoverHitTest_Implementation(FInputDeviceRay HoverPos, const FScriptableToolModifierStates& Modifiers);
+	UE_API virtual FInputRayHit OnHoverHitTest_Implementation(FInputDeviceRay HoverPos, const FScriptableToolModifierStates& Modifiers);
 
 	/**
 	 * Begin a sequence of Hover input events 
@@ -99,9 +101,9 @@ public:
 	 * @param Modifiers current modifier key/button state
 	 */
 	UFUNCTION(BlueprintNativeEvent, Category = "Hover Events")
-	void OnHoverBegin(FInputDeviceRay HoverPos, const FScriptableToolModifierStates& Modifiers);
+	UE_API void OnHoverBegin(FInputDeviceRay HoverPos, const FScriptableToolModifierStates& Modifiers);
 
-	virtual void OnHoverBegin_Implementation(FInputDeviceRay HoverPos, const FScriptableToolModifierStates& Modifiers);
+	UE_API virtual void OnHoverBegin_Implementation(FInputDeviceRay HoverPos, const FScriptableToolModifierStates& Modifiers);
 
 	/**
 	 * Update an active Hover sequence. OnHoverUpdate is only ever called between OnHoverBegin and OnHoverEnd
@@ -110,22 +112,22 @@ public:
 	 * @return true to continue hovering, false to stop receiving additional hover events
 	 */
 	UFUNCTION(BlueprintNativeEvent, Category = "Hover Events")
-	UPARAM(DisplayName="Continue Hovering") bool
+	UE_API UPARAM(DisplayName="Continue Hovering") bool
 	OnHoverUpdate(FInputDeviceRay HoverPos, const FScriptableToolModifierStates& Modifiers);
 
-	virtual bool OnHoverUpdate_Implementation(FInputDeviceRay HoverPos, const FScriptableToolModifierStates& Modifiers);
+	UE_API virtual bool OnHoverUpdate_Implementation(FInputDeviceRay HoverPos, const FScriptableToolModifierStates& Modifiers);
 
 	/**
 	 * Indicate that an active captured Hover sequence has ended. The device may no longer be in the viewport, so no position can be provided.
 	 * @param Modifiers current modifier key/button state
 	 */
 	UFUNCTION(BlueprintNativeEvent, Category = "Hover Events")
-	void OnHoverEnd(const FScriptableToolModifierStates& Modifiers);
+	UE_API void OnHoverEnd(const FScriptableToolModifierStates& Modifiers);
 
-	virtual void OnHoverEnd_Implementation(const FScriptableToolModifierStates& Modifiers);
+	UE_API virtual void OnHoverEnd_Implementation(const FScriptableToolModifierStates& Modifiers);
 
 	UFUNCTION(BlueprintCallable, Category = "ScriptableTool|Input")
-	UPARAM(DisplayName="In Hover State") bool InActiveHover() const;
+	UE_API UPARAM(DisplayName="In Hover State") bool InActiveHover() const;
 
 	//
 	// Modifer Buttons Support API
@@ -134,24 +136,24 @@ public:
 
 	/** @return true if the Shift key is currently held down */
 	UFUNCTION(BlueprintCallable, Category="ScriptableTool|Input")
-	UPARAM(DisplayName="Shift Down") bool IsShiftDown() const;
+	UE_API UPARAM(DisplayName="Shift Down") bool IsShiftDown() const;
 
 	/** @return true if the Ctrl key is currently held down */
 	UFUNCTION(BlueprintCallable, Category="ScriptableTool|Input")
-	UPARAM(DisplayName="Ctrl Down") bool IsCtrlDown() const;
+	UE_API UPARAM(DisplayName="Ctrl Down") bool IsCtrlDown() const;
 
 	/** @return true if the Alt key is currently held down */
 	UFUNCTION(BlueprintCallable, Category="ScriptableTool|Input")
-	UPARAM(DisplayName="Alt Down") bool IsAltDown() const;
+	UE_API UPARAM(DisplayName="Alt Down") bool IsAltDown() const;
 
 	/** @return a struct containing the current Modifier key states */
 	UFUNCTION(BlueprintCallable, Category="ScriptableTool|Input")
-	FScriptableToolModifierStates GetActiveModifiers();
+	UE_API FScriptableToolModifierStates GetActiveModifiers();
 
 
 
 public:
-	virtual void Setup() override;
+	UE_API virtual void Setup() override;
 
 	UPROPERTY(Transient, DuplicateTransient, NonTransactional, SkipSerialization)
 	TObjectPtr<USingleClickInputBehavior> SingleClickBehavior;
@@ -166,13 +168,15 @@ public:
 	bool bAltModifier = false;
 
 	// IClickBehaviorTarget API
-	virtual FInputRayHit IsHitByClick(const FInputDeviceRay& ClickPos);
-	virtual void OnClicked(const FInputDeviceRay& ClickPos);
-	virtual void OnUpdateModifierState(int ModifierID, bool bIsOn);
+	UE_API virtual FInputRayHit IsHitByClick(const FInputDeviceRay& ClickPos);
+	UE_API virtual void OnClicked(const FInputDeviceRay& ClickPos);
+	UE_API virtual void OnUpdateModifierState(int ModifierID, bool bIsOn);
 
 	// IHoverBehaviorTarget API
-	virtual FInputRayHit BeginHoverSequenceHitTest(const FInputDeviceRay& PressPos);
-	virtual void OnBeginHover(const FInputDeviceRay& DevicePos);
-	virtual bool OnUpdateHover(const FInputDeviceRay& DevicePos);
-	virtual void OnEndHover();
-};
\ No newline at end of file
+	UE_API virtual FInputRayHit BeginHoverSequenceHitTest(const FInputDeviceRay& PressPos);
+	UE_API virtual void OnBeginHover(const FInputDeviceRay& DevicePos);
+	UE_API virtual bool OnUpdateHover(const FInputDeviceRay& DevicePos);
+	UE_API virtual void OnEndHover();
+};
+
+#undef UE_API
diff --git a/Engine/Plugins/Runtime/ScriptableToolsFramework/Source/ScriptableToolsFramework/Public/Drawing/ScriptableToolLine.h b/Engine/Plugins/Runtime/ScriptableToolsFramework/Source/ScriptableToolsFramework/Public/Drawing/ScriptableToolLine.h
--- a/Engine/Plugins/Runtime/ScriptableToolsFramework/Source/ScriptableToolsFramework/Public/Drawing/ScriptableToolLine.h
+++ b/Engine/Plugins/Runtime/ScriptableToolsFramework/Source/ScriptableToolsFramework/Public/Drawing/ScriptableToolLine.h
@@ -9,33 +9,35 @@
 
 #include "ScriptableToolLine.generated.h"
 
+#define UE_API SCRIPTABLETOOLSFRAMEWORK_API
 
-UCLASS(BlueprintType)
-class SCRIPTABLETOOLSFRAMEWORK_API UScriptableToolLine : public UObject
+
+UCLASS(MinimalAPI, BlueprintType)
+class UScriptableToolLine : public UObject
 {
 	GENERATED_BODY()
 
 public:
 
-	void SetLineID(int32 LineIDIn);
-	int32 GetLineID() const;
+	UE_API void SetLineID(int32 LineIDIn);
+	UE_API int32 GetLineID() const;
 
-	bool IsDirty() const;
-	FRenderableLine GenerateLineDescription();
+	UE_API bool IsDirty() const;
+	UE_API FRenderableLine GenerateLineDescription();
 
 	/**
 	 * Set the starting position of the line
 	 * @param Start The position in space of the start of the line
 	 */
 	UFUNCTION(BlueprintCallable, Category = "ScriptableTool|Drawing|Lines")
-	void SetLineStart(FVector Start);
+	UE_API void SetLineStart(FVector Start);
 
 	/**
 	 * Set the ending position of the line
 	 * @param End The position in space of the end of the line
 	 */
 	UFUNCTION(BlueprintCallable, Category = "ScriptableTool|Drawing|Lines")
-	void SetLineEnd(FVector End);
+	UE_API void SetLineEnd(FVector End);
 
 	/**
 	 * Set the starting and ending positions of the line
@@ -43,21 +45,21 @@ public:
 	 * @param End The position in space of the end of the line
 	 */
 	UFUNCTION(BlueprintCallable, Category = "ScriptableTool|Drawing|Lines")
-	void SetLineEndPoints(FVector Start, FVector End);
+	UE_API void SetLineEndPoints(FVector Start, FVector End);
 
 	/**
 	 * Set the line's color
 	 * @param Color The color the line should be rendered as
 	 */
 	UFUNCTION(BlueprintCallable, Category = "ScriptableTool|Drawing|Lines")
-	void SetLineColor(FColor Color);
+	UE_API void SetLineColor(FColor Color);
 
 	/**
 	 * Set the line's thickness
 	 * @param Thickness The visual thickness in pixels the line should be rendered as
 	 */
 	UFUNCTION(BlueprintCallable, Category = "ScriptableTool|Drawing|Lines")
-	void SetLineThickness(float Thickness);
+	UE_API void SetLineThickness(float Thickness);
 
 	/**
 	 * Set the line's depth bias. The depth bias controls small micro adjustments in effective
@@ -66,7 +68,7 @@ public:
 	 * @param DepthBias The depth bias adjustment for the line
 	 */
 	UFUNCTION(BlueprintCallable, Category = "ScriptableTool|Drawing|Lines")
-	void SetLineDepthBias(float DepthBias);
+	UE_API void SetLineDepthBias(float DepthBias);
 
 private:
 
@@ -78,4 +80,6 @@ private:
 	UPROPERTY()
 	int32 LineID;
 
-};
\ No newline at end of file
+};
+
+#undef UE_API
diff --git a/Engine/Plugins/Runtime/ScriptableToolsFramework/Source/ScriptableToolsFramework/Public/Drawing/ScriptableToolLineSet.h b/Engine/Plugins/Runtime/ScriptableToolsFramework/Source/ScriptableToolsFramework/Public/Drawing/ScriptableToolLineSet.h
--- a/Engine/Plugins/Runtime/ScriptableToolsFramework/Source/ScriptableToolsFramework/Public/Drawing/ScriptableToolLineSet.h
+++ b/Engine/Plugins/Runtime/ScriptableToolsFramework/Source/ScriptableToolsFramework/Public/Drawing/ScriptableToolLineSet.h
@@ -7,20 +7,22 @@
 
 #include "ScriptableToolLineSet.generated.h"
 
+#define UE_API SCRIPTABLETOOLSFRAMEWORK_API
+
 
 class UScriptableToolLine;
 class UPreviewGeometry;
 
-UCLASS(BlueprintType)
-class SCRIPTABLETOOLSFRAMEWORK_API UScriptableToolLineSet : public UObject
+UCLASS(MinimalAPI, BlueprintType)
+class UScriptableToolLineSet : public UObject
 {
 	GENERATED_BODY()
 
 public:
 
-	void Initialize(TObjectPtr<UPreviewGeometry> PreviewGeometry);
+	UE_API void Initialize(TObjectPtr<UPreviewGeometry> PreviewGeometry);
 
-	void OnTick();
+	UE_API void OnTick();
 
 
 	/**
@@ -28,33 +30,33 @@ public:
 	 * @return The new line object added to the set.
 	 */
 	UFUNCTION(BlueprintCallable, Category = "ScriptableTool|Drawing|LineSet")
-	UPARAM(DisplayName = "LineComponent") UScriptableToolLine* AddLine();
+	UE_API UPARAM(DisplayName = "LineComponent") UScriptableToolLine* AddLine();
 
 	/**
 	 * Remove a specific line object from the set, removing it fromt the scene.
 	 * @param Line A reference to a line to be removed from the set.
 	 */
 	UFUNCTION(BlueprintCallable, Category = "ScriptableTool|Drawing|LineSet")
-	void RemoveLine(UScriptableToolLine* Line);
+	UE_API void RemoveLine(UScriptableToolLine* Line);
 
 
 	/**
 	 * Remove all current lines in the set.
 	 */
 	UFUNCTION(BlueprintCallable, Category = "ScriptableTool|Drawing|LineSet")
-	void RemoveAllLines();
+	UE_API void RemoveAllLines();
 
 	/**
 	 * Set the color of all lines in the set simultaneously.
 	 */
 	UFUNCTION(BlueprintCallable, Category = "ScriptableTool|Drawing|LineSet")
-	void SetAllLinesColor(FColor Color);
+	UE_API void SetAllLinesColor(FColor Color);
 
 	/**
 	 * Set the thickness of all lines in the set simultaneously.
 	 */
 	UFUNCTION(BlueprintCallable, Category = "ScriptableTool|Drawing|LineSet")
-	void SetAllLinesThickness(float Thickness);
+	UE_API void SetAllLinesThickness(float Thickness);
 
 protected:
 
@@ -69,3 +71,4 @@ protected:
 
 };
 
+#undef UE_API
diff --git a/Engine/Plugins/Runtime/ScriptableToolsFramework/Source/ScriptableToolsFramework/Public/Drawing/ScriptableToolPoint.h b/Engine/Plugins/Runtime/ScriptableToolsFramework/Source/ScriptableToolsFramework/Public/Drawing/ScriptableToolPoint.h
--- a/Engine/Plugins/Runtime/ScriptableToolsFramework/Source/ScriptableToolsFramework/Public/Drawing/ScriptableToolPoint.h
+++ b/Engine/Plugins/Runtime/ScriptableToolsFramework/Source/ScriptableToolsFramework/Public/Drawing/ScriptableToolPoint.h
@@ -9,40 +9,42 @@
 
 #include "ScriptableToolPoint.generated.h"
 
+#define UE_API SCRIPTABLETOOLSFRAMEWORK_API
 
-UCLASS(BlueprintType)
-class SCRIPTABLETOOLSFRAMEWORK_API UScriptableToolPoint : public UObject
+
+UCLASS(MinimalAPI, BlueprintType)
+class UScriptableToolPoint : public UObject
 {
 	GENERATED_BODY()
 
 public:
 
-	void SetPointID(int32 PointIDIn);
-	int32 GetPointID() const;
+	UE_API void SetPointID(int32 PointIDIn);
+	UE_API int32 GetPointID() const;
 
-	bool IsDirty() const;
-	FRenderablePoint GeneratePointDescription();
+	UE_API bool IsDirty() const;
+	UE_API FRenderablePoint GeneratePointDescription();
 
 	/**
 	 * Set the position of the point
 	 * @param Position The position in space of the point
 	 */
 	UFUNCTION(BlueprintCallable, Category = "ScriptableTool|Drawing|Points")
-	void SetPointPosition(FVector Position);
+	UE_API void SetPointPosition(FVector Position);
 
 	/**
 	 * Set the point's color
 	 * @param Color The color the point should be rendered as
 	 */
 	UFUNCTION(BlueprintCallable, Category = "ScriptableTool|Drawing|Points")
-	void SetPointColor(FColor Color);
+	UE_API void SetPointColor(FColor Color);
 
 	/**
 	 * Set the point's size
 	 * @param Size The visual size in pixels the point should be rendered as
 	 */
 	UFUNCTION(BlueprintCallable, Category = "ScriptableTool|Drawing|Points")
-	void SetPointSize(float Thickness);
+	UE_API void SetPointSize(float Thickness);
 
 	/**
 	 * Set the point's depth bias. The depth bias controls small micro adjustments in effective
@@ -51,7 +53,7 @@ public:
 	 * @param DepthBias The depth bias adjustment for the point
 	 */
 	UFUNCTION(BlueprintCallable, Category = "ScriptableTool|Drawing|Points")
-	void SetPointDepthBias(float DepthBias);
+	UE_API void SetPointDepthBias(float DepthBias);
 
 private:
 
@@ -63,4 +65,6 @@ private:
 	UPROPERTY()
 	int32 PointID;
 
-};
\ No newline at end of file
+};
+
+#undef UE_API
diff --git a/Engine/Plugins/Runtime/ScriptableToolsFramework/Source/ScriptableToolsFramework/Public/Drawing/ScriptableToolPointSet.h b/Engine/Plugins/Runtime/ScriptableToolsFramework/Source/ScriptableToolsFramework/Public/Drawing/ScriptableToolPointSet.h
--- a/Engine/Plugins/Runtime/ScriptableToolsFramework/Source/ScriptableToolsFramework/Public/Drawing/ScriptableToolPointSet.h
+++ b/Engine/Plugins/Runtime/ScriptableToolsFramework/Source/ScriptableToolsFramework/Public/Drawing/ScriptableToolPointSet.h
@@ -7,27 +7,29 @@
 
 #include "ScriptableToolPointSet.generated.h"
 
+#define UE_API SCRIPTABLETOOLSFRAMEWORK_API
+
 
 class UScriptableToolPoint;
 class UPreviewGeometry;
 
-UCLASS(BlueprintType)
-class SCRIPTABLETOOLSFRAMEWORK_API UScriptableToolPointSet : public UObject
+UCLASS(MinimalAPI, BlueprintType)
+class UScriptableToolPointSet : public UObject
 {
 	GENERATED_BODY()
 
 public:
 
-	void Initialize(TObjectPtr<UPreviewGeometry> PreviewGeometry);
+	UE_API void Initialize(TObjectPtr<UPreviewGeometry> PreviewGeometry);
 
-	void OnTick();
+	UE_API void OnTick();
 
 	/**
 	 * Create and return a new point object. Users should save a reference to this object for future updates or removal from the set.
 	 * @return The new point object added to the set.
 	 */
 	UFUNCTION(BlueprintCallable, Category = "ScriptableTool|Drawing|PointSet")
-	UPARAM(DisplayName = "PointComponent") UScriptableToolPoint* AddPoint();
+	UE_API UPARAM(DisplayName = "PointComponent") UScriptableToolPoint* AddPoint();
 
 
 	/**
@@ -35,25 +37,25 @@ public:
 	 * @param Point A reference to a point to be removed from the set.
 	 */
 	UFUNCTION(BlueprintCallable, Category = "ScriptableTool|Drawing|PointSet")
-	void RemovePoint(UScriptableToolPoint* Point);
+	UE_API void RemovePoint(UScriptableToolPoint* Point);
 
 	/**
 	 * Remove all current points in the set.
 	 */
 	UFUNCTION(BlueprintCallable, Category = "ScriptableTool|Drawing|PointSet")
-	void RemoveAllPoints();
+	UE_API void RemoveAllPoints();
 
 	/**
 	 * Set the color of all points in the set simultaneously.
 	 */
 	UFUNCTION(BlueprintCallable, Category = "ScriptableTool|Drawing|PointSet")
-	void SetAllPointsColor(FColor Color);
+	UE_API void SetAllPointsColor(FColor Color);
 
 	/**
 	 * Set the size of all points in the set simultaneously.
 	 */
 	UFUNCTION(BlueprintCallable, Category = "ScriptableTool|Drawing|PointSet")
-	void SetAllPointsSize(float Size);
+	UE_API void SetAllPointsSize(float Size);
 
 protected:
 
@@ -68,3 +70,4 @@ protected:
 
 };
 
+#undef UE_API
diff --git a/Engine/Plugins/Runtime/ScriptableToolsFramework/Source/ScriptableToolsFramework/Public/Drawing/ScriptableToolTriangle.h b/Engine/Plugins/Runtime/ScriptableToolsFramework/Source/ScriptableToolsFramework/Public/Drawing/ScriptableToolTriangle.h
--- a/Engine/Plugins/Runtime/ScriptableToolsFramework/Source/ScriptableToolsFramework/Public/Drawing/ScriptableToolTriangle.h
+++ b/Engine/Plugins/Runtime/ScriptableToolsFramework/Source/ScriptableToolsFramework/Public/Drawing/ScriptableToolTriangle.h
@@ -9,21 +9,23 @@
 
 #include "ScriptableToolTriangle.generated.h"
 
+#define UE_API SCRIPTABLETOOLSFRAMEWORK_API
 
-UCLASS(BlueprintType)
-class SCRIPTABLETOOLSFRAMEWORK_API UScriptableToolTriangle : public UObject
+
+UCLASS(MinimalAPI, BlueprintType)
+class UScriptableToolTriangle : public UObject
 {
 	GENERATED_BODY()
 
 public:
 
-	UScriptableToolTriangle();
+	UE_API UScriptableToolTriangle();
 
-	void SetTriangleID(int32 TriangleIDIn);
-	int32 GetTriangleID() const;
+	UE_API void SetTriangleID(int32 TriangleIDIn);
+	UE_API int32 GetTriangleID() const;
 
-	bool IsDirty() const;
-	FRenderableTriangle GenerateTriangleDescription();
+	UE_API bool IsDirty() const;
+	UE_API FRenderableTriangle GenerateTriangleDescription();
 
 	/*
 	* Set the material of the triangle
@@ -31,7 +33,7 @@ public:
 	* @param Material The new material that should be assigned to the triangle
 	*/
 	UFUNCTION(BlueprintCallable, Category = "ScriptableTool|Drawing|Triangles")
-	void SetTriangleMaterial(UMaterialInterface* Material);
+	UE_API void SetTriangleMaterial(UMaterialInterface* Material);
 
 	/*
 	* Set the points of the triangle
@@ -41,7 +43,7 @@ public:
 	* @param C The position of the third corner of the triangle
 	*/
 	UFUNCTION(BlueprintCallable, Category = "ScriptableTool|Drawing|Triangles")
-	void SetTrianglePoints(FVector A, FVector B, FVector C);
+	UE_API void SetTrianglePoints(FVector A, FVector B, FVector C);
 
 	/*
 	* Set the UV coordinates of the triangle
@@ -51,7 +53,7 @@ public:
 	* @param C The UV coordinate of the third corner of the triangle
 	*/
 	UFUNCTION(BlueprintCallable, Category = "ScriptableTool|Drawing|Triangles")
-	void SetTriangleUVs(FVector2D A, FVector2D B, FVector2D C);
+	UE_API void SetTriangleUVs(FVector2D A, FVector2D B, FVector2D C);
 
 	/*
 	* Set the normal values of the triangle
@@ -61,7 +63,7 @@ public:
 	* @param C The normal of the third corner of the triangle
 	*/
 	UFUNCTION(BlueprintCallable, Category = "ScriptableTool|Drawing|Triangles")
-	void SetTriangleNormals(FVector A, FVector B, FVector C);
+	UE_API void SetTriangleNormals(FVector A, FVector B, FVector C);
 
 
 	/*
@@ -72,7 +74,7 @@ public:
 	* @param C The vertex color of the third corner of the triangle
 	*/
 	UFUNCTION(BlueprintCallable, Category = "ScriptableTool|Drawing|Triangles")
-	void SetTriangleColors(FColor A, FColor B, FColor C);
+	UE_API void SetTriangleColors(FColor A, FColor B, FColor C);
 
 private:
 
@@ -86,29 +88,29 @@ private:
 
 };
 
-UCLASS(BlueprintType)
-class SCRIPTABLETOOLSFRAMEWORK_API UScriptableToolQuad : public UObject
+UCLASS(MinimalAPI, BlueprintType)
+class UScriptableToolQuad : public UObject
 {
 	GENERATED_BODY()
 
 public:
 
-	UScriptableToolQuad();
+	UE_API UScriptableToolQuad();
 
-	void SetTriangleAID(int32 TriangleIDIn);
-	int32 GetTriangleAID() const;
-	void SetTriangleBID(int32 TriangleIDIn);
-	int32 GetTriangleBID() const;
+	UE_API void SetTriangleAID(int32 TriangleIDIn);
+	UE_API int32 GetTriangleAID() const;
+	UE_API void SetTriangleBID(int32 TriangleIDIn);
+	UE_API int32 GetTriangleBID() const;
 
-	bool IsDirty() const;
-	TPair<FRenderableTriangle, FRenderableTriangle> GenerateQuadDescription();
+	UE_API bool IsDirty() const;
+	UE_API TPair<FRenderableTriangle, FRenderableTriangle> GenerateQuadDescription();
 
 	/**
 	 * Set the material of the quad
 	 * @param Material The new material that should be assigned to the quad
 	 */
 	UFUNCTION(BlueprintCallable, Category = "ScriptableTool|Drawing|Quads")
-	void SetQuadMaterial(UMaterialInterface* Material);
+	UE_API void SetQuadMaterial(UMaterialInterface* Material);
 
 	/**
 	 * Set the points of the quad
@@ -118,7 +120,7 @@ public:
 	 * @param D The position of the fourth corner of the quad
 	 */
 	UFUNCTION(BlueprintCallable, Category = "ScriptableTool|Drawing|Quads")
-	void SetQuadPoints(FVector A, FVector B, FVector C, FVector D);
+	UE_API void SetQuadPoints(FVector A, FVector B, FVector C, FVector D);
 
 	/**
 	 * Set the UV coordinates of the quad
@@ -128,7 +130,7 @@ public:
 	 * @param D The UV coordinate of the fourth corner of the quad
 	 */
 	UFUNCTION(BlueprintCallable, Category = "ScriptableTool|Drawing|Quads")
-	void SetQuadUVs(FVector2D A, FVector2D B, FVector2D C, FVector2D D);
+	UE_API void SetQuadUVs(FVector2D A, FVector2D B, FVector2D C, FVector2D D);
 
 	/**
 	 * Set the normal values of the quad
@@ -138,7 +140,7 @@ public:
 	 * @param D The normal of the fourth corner of the quad
 	 */
 	UFUNCTION(BlueprintCallable, Category = "ScriptableTool|Drawing|Quads")
-	void SetQuadNormals(FVector A, FVector B, FVector C, FVector D);
+	UE_API void SetQuadNormals(FVector A, FVector B, FVector C, FVector D);
 
 	/**
 	 * Set the vertex colors of the quad
@@ -148,7 +150,7 @@ public:
 	 * @param D The vertex color of the fourth corner of the quad
 	 */
 	UFUNCTION(BlueprintCallable, Category = "ScriptableTool|Drawing|Quads")
-	void SetQuadColors(FColor A, FColor B, FColor C, FColor D);
+	UE_API void SetQuadColors(FColor A, FColor B, FColor C, FColor D);
 private:
 
 	FRenderableTriangle TriangleADescription;
@@ -163,4 +165,6 @@ private:
 	UPROPERTY()
 	int32 TriangleBID;
 
-};
\ No newline at end of file
+};
+
+#undef UE_API
diff --git a/Engine/Plugins/Runtime/ScriptableToolsFramework/Source/ScriptableToolsFramework/Public/Drawing/ScriptableToolTriangleSet.h b/Engine/Plugins/Runtime/ScriptableToolsFramework/Source/ScriptableToolsFramework/Public/Drawing/ScriptableToolTriangleSet.h
--- a/Engine/Plugins/Runtime/ScriptableToolsFramework/Source/ScriptableToolsFramework/Public/Drawing/ScriptableToolTriangleSet.h
+++ b/Engine/Plugins/Runtime/ScriptableToolsFramework/Source/ScriptableToolsFramework/Public/Drawing/ScriptableToolTriangleSet.h
@@ -7,67 +7,69 @@
 
 #include "ScriptableToolTriangleSet.generated.h"
 
+#define UE_API SCRIPTABLETOOLSFRAMEWORK_API
+
 
 class UScriptableToolTriangle;
 class UScriptableToolQuad;
 class UPreviewGeometry;
 
-UCLASS(BlueprintType)
-class SCRIPTABLETOOLSFRAMEWORK_API UScriptableToolTriangleSet : public UObject
+UCLASS(MinimalAPI, BlueprintType)
+class UScriptableToolTriangleSet : public UObject
 {
 	GENERATED_BODY()
 
 public:
 
-	void Initialize(TObjectPtr<UPreviewGeometry> PreviewGeometry);
+	UE_API void Initialize(TObjectPtr<UPreviewGeometry> PreviewGeometry);
 
-	void OnTick();
+	UE_API void OnTick();
 
 	/**
 	 * Create and return a new triangle object. Users should save a reference to this object for future updates or removal from the set.
 	 * @return The new triangle object added to the set
 	 */
 	UFUNCTION(BlueprintCallable, Category = "ScriptableTool|Drawing|TriangleSet")
-	UPARAM(DisplayName = "TriangleComponent") UScriptableToolTriangle* AddTriangle();
+	UE_API UPARAM(DisplayName = "TriangleComponent") UScriptableToolTriangle* AddTriangle();
 
 	/**
 	 * Create and return a new quad object. Quad objects are two paired triangles. Users should save a reference to this object for future updates or removal from the set.
 	 * @return The new quad object added to the set
 	 */
 	UFUNCTION(BlueprintCallable, Category = "ScriptableTool|Drawing|TriangleSet")
-	UPARAM(DisplayName = "QuadComponent") UScriptableToolQuad* AddQuad();
+	UE_API UPARAM(DisplayName = "QuadComponent") UScriptableToolQuad* AddQuad();
 
 	/**
 	 * Remove a specific triangle object from the set, removing it from the scene.
 	 * @param Triangle A reference to a triangle to be removed from the set.
 	 */
 	UFUNCTION(BlueprintCallable, Category = "ScriptableTool|Drawing|TriangleSet")
-	void RemoveTriangle(UScriptableToolTriangle* Triangle);
+	UE_API void RemoveTriangle(UScriptableToolTriangle* Triangle);
 
 	/**
 	 * Remove a specific quad object from the set, removing it from the scene.
 	 * @param Quad A reference to a quad to be removed from the set.
 	 */
 	UFUNCTION(BlueprintCallable, Category = "ScriptableTool|Drawing|TriangleSet")
-	void RemoveQuad(UScriptableToolQuad* Quad);
+	UE_API void RemoveQuad(UScriptableToolQuad* Quad);
 
 	/**
 	 * Remove all current triangles and quads in the set. 
 	 */
 	UFUNCTION(BlueprintCallable, Category = "ScriptableTool|Drawing|TriangleSet")
-	void RemoveAllFaces();
+	UE_API void RemoveAllFaces();
 
 	/**
 	 * Set the color of all triangles and quads in the set simultaneously .
 	 */
 	UFUNCTION(BlueprintCallable, Category = "ScriptableTool|Drawing|TriangleSet")
-	void SetAllTrianglesColor(FColor Color);
+	UE_API void SetAllTrianglesColor(FColor Color);
 
 	/**
 	 * Set the material of all triangles and quads in the set simultaneously .
 	 */
 	UFUNCTION(BlueprintCallable, Category = "ScriptableTool|Drawing|TriangleSet")
-	void SetAllTrianglesMaterial(UMaterialInterface* Material);
+	UE_API void SetAllTrianglesMaterial(UMaterialInterface* Material);
 
 protected:
 
@@ -85,3 +87,4 @@ protected:
 
 };
 
+#undef UE_API
diff --git a/Engine/Plugins/Runtime/ScriptableToolsFramework/Source/ScriptableToolsFramework/Public/ScriptableToolBuilder.h b/Engine/Plugins/Runtime/ScriptableToolsFramework/Source/ScriptableToolsFramework/Public/ScriptableToolBuilder.h
--- a/Engine/Plugins/Runtime/ScriptableToolsFramework/Source/ScriptableToolsFramework/Public/ScriptableToolBuilder.h
+++ b/Engine/Plugins/Runtime/ScriptableToolsFramework/Source/ScriptableToolsFramework/Public/ScriptableToolBuilder.h
@@ -5,6 +5,8 @@
 #include "InteractiveToolBuilder.h"
 #include "ScriptableToolBuilder.generated.h"
 
+#define UE_API SCRIPTABLETOOLSFRAMEWORK_API
+
 
 namespace ScriptableToolBuilderHelpers
 {
@@ -17,16 +19,16 @@ class UScriptableInteractiveTool;
  * UBaseScriptableToolBuilder is a trivial base UInteractiveToolBuilder for any UScriptableInteractiveTool subclass.
  * CanBuildTool will return true as long as the ToolClass is a valid UClass.
  */
-UCLASS()
-class SCRIPTABLETOOLSFRAMEWORK_API UBaseScriptableToolBuilder : public UInteractiveToolBuilder
+UCLASS(MinimalAPI)
+class UBaseScriptableToolBuilder : public UInteractiveToolBuilder
 {
 	GENERATED_BODY()
 public:
 
 	TWeakObjectPtr<UClass> ToolClass;
 
-	virtual bool CanBuildTool(const FToolBuilderState& SceneState) const override;
-	virtual UInteractiveTool* BuildTool(const FToolBuilderState& SceneState) const override;
+	UE_API virtual bool CanBuildTool(const FToolBuilderState& SceneState) const override;
+	UE_API virtual UInteractiveTool* BuildTool(const FToolBuilderState& SceneState) const override;
 };
 
 UINTERFACE(MinimalAPI, BlueprintType)
@@ -45,8 +47,8 @@ public:
 	virtual void SetupTool(const FToolBuilderState& SceneState, UInteractiveTool* Tool) const = 0;
 };
 
-UCLASS()
-class SCRIPTABLETOOLSFRAMEWORK_API UCustomScriptableToolBuilderComponentBase : public UObject
+UCLASS(MinimalAPI)
+class UCustomScriptableToolBuilderComponentBase : public UObject
 {
 	GENERATED_BODY()
 
@@ -55,18 +57,18 @@ public:
 };
 
 
-UCLASS(Transient, Blueprintable, Hidden)
-class SCRIPTABLETOOLSFRAMEWORK_API UCustomScriptableToolBuilderContainer : public UBaseScriptableToolBuilder
+UCLASS(MinimalAPI, Transient, Blueprintable, Hidden)
+class UCustomScriptableToolBuilderContainer : public UBaseScriptableToolBuilder
 {
 	GENERATED_BODY()
 
 
 public:
 
-	void Initialize(TObjectPtr<UCustomScriptableToolBuilderComponentBase> BuilderInstanceIn);
+	UE_API void Initialize(TObjectPtr<UCustomScriptableToolBuilderComponentBase> BuilderInstanceIn);
 
-	virtual bool CanBuildTool(const FToolBuilderState& SceneState) const override;
-	virtual UInteractiveTool* BuildTool(const FToolBuilderState& SceneState) const override;
+	UE_API virtual bool CanBuildTool(const FToolBuilderState& SceneState) const override;
+	UE_API virtual UInteractiveTool* BuildTool(const FToolBuilderState& SceneState) const override;
 
 private:
 
@@ -80,25 +82,25 @@ private:
 *
 */
 
-UCLASS(Transient, Blueprintable, Abstract)
-class SCRIPTABLETOOLSFRAMEWORK_API UCustomScriptableToolBuilder : public UCustomScriptableToolBuilderComponentBase, public ICustomScriptableToolBuilderBaseInterface
+UCLASS(MinimalAPI, Transient, Blueprintable, Abstract)
+class UCustomScriptableToolBuilder : public UCustomScriptableToolBuilderComponentBase, public ICustomScriptableToolBuilderBaseInterface
 {
 	GENERATED_BODY()
 
 public:
 
 	UFUNCTION(BlueprintNativeEvent, Category = "ScriptableToolBuilder|Events")
-	bool OnCanBuildTool(const TArray<AActor*>& SelectedActors, const TArray<UActorComponent*>& SelectedComponents) const;
+	UE_API bool OnCanBuildTool(const TArray<AActor*>& SelectedActors, const TArray<UActorComponent*>& SelectedComponents) const;
 
-	bool OnCanBuildTool_Implementation(const TArray<AActor*>& SelectedActors, const TArray<UActorComponent*>& SelectedComponents) const;
+	UE_API bool OnCanBuildTool_Implementation(const TArray<AActor*>& SelectedActors, const TArray<UActorComponent*>& SelectedComponents) const;
 
 	UFUNCTION(BlueprintNativeEvent, Category = "ScriptableToolBuilder|Events")
-	void OnSetupTool(UScriptableInteractiveTool* Tool, const TArray<AActor*>& SelectedActors, const TArray<UActorComponent*>& SelectedComponents) const;
+	UE_API void OnSetupTool(UScriptableInteractiveTool* Tool, const TArray<AActor*>& SelectedActors, const TArray<UActorComponent*>& SelectedComponents) const;
 
-	void OnSetupTool_Implementation(UScriptableInteractiveTool* Tool, const TArray<AActor*>& SelectedActors, const TArray<UActorComponent*>& SelectedComponents) const;
+	UE_API void OnSetupTool_Implementation(UScriptableInteractiveTool* Tool, const TArray<AActor*>& SelectedActors, const TArray<UActorComponent*>& SelectedComponents) const;
 
-	virtual bool CanBuildTool(const FToolBuilderState& SceneState) const override;
-	virtual void SetupTool(const FToolBuilderState& SceneState, UInteractiveTool* Tool) const override;
+	UE_API virtual bool CanBuildTool(const FToolBuilderState& SceneState) const override;
+	UE_API virtual void SetupTool(const FToolBuilderState& SceneState, UInteractiveTool* Tool) const override;
 };
 
 
@@ -108,14 +110,14 @@ public:
 * 
 */
 
-UCLASS(Blueprintable)
-class SCRIPTABLETOOLSFRAMEWORK_API UScriptableToolTargetRequirements : public UObject
+UCLASS(MinimalAPI, Blueprintable)
+class UScriptableToolTargetRequirements : public UObject
 {
 	GENERATED_BODY()
 public:
 
 	UFUNCTION(BlueprintCallable, Category = "ScriptableToolBuilder|ToolTargets")
-	static UPARAM(DisplayName = "New Target Requirements") UScriptableToolTargetRequirements*
+	static UE_API UPARAM(DisplayName = "New Target Requirements") UScriptableToolTargetRequirements*
 	BuildToolTargetRequirements(TArray<UClass*> RequirementInterfaces);
 
 	const FToolTargetTypeRequirements& GetRequirements() const { return Requirements; };
@@ -132,31 +134,33 @@ private:
 
 };
 
-UCLASS(Transient, Blueprintable, Abstract)
-class SCRIPTABLETOOLSFRAMEWORK_API UToolTargetScriptableToolBuilder: public UCustomScriptableToolBuilderComponentBase, public ICustomScriptableToolBuilderBaseInterface
+UCLASS(MinimalAPI, Transient, Blueprintable, Abstract)
+class UToolTargetScriptableToolBuilder: public UCustomScriptableToolBuilderComponentBase, public ICustomScriptableToolBuilderBaseInterface
 {
 	GENERATED_BODY()
 
 public:
 
-	void Initialize();
+	UE_API void Initialize();
 
 	UFUNCTION(BlueprintNativeEvent, Category = "Tool Targets")
-	UScriptableToolTargetRequirements* GetToolTargetRequirements() const;
+	UE_API UScriptableToolTargetRequirements* GetToolTargetRequirements() const;
 
-	virtual UScriptableToolTargetRequirements* GetToolTargetRequirements_Implementation() const;
+	UE_API virtual UScriptableToolTargetRequirements* GetToolTargetRequirements_Implementation() const;
 
 	UFUNCTION(BlueprintNativeEvent, Category = "ScriptableToolBuilder|Events")
-	void OnSetupTool(UScriptableInteractiveTool* Tool) const;
+	UE_API void OnSetupTool(UScriptableInteractiveTool* Tool) const;
 
-	void OnSetupTool_Implementation(UScriptableInteractiveTool* Tool) const;
+	UE_API void OnSetupTool_Implementation(UScriptableInteractiveTool* Tool) const;
 
-	virtual bool CanBuildTool(const FToolBuilderState& SceneState) const override;
-	virtual void SetupTool(const FToolBuilderState& SceneState, UInteractiveTool* Tool) const override;
+	UE_API virtual bool CanBuildTool(const FToolBuilderState& SceneState) const override;
+	UE_API virtual void SetupTool(const FToolBuilderState& SceneState, UInteractiveTool* Tool) const override;
 	
 private:
 
 	UPROPERTY()
 	TObjectPtr<UScriptableToolTargetRequirements> Requirements;
 
-};
\ No newline at end of file
+};
+
+#undef UE_API
diff --git a/Engine/Plugins/Runtime/ScriptableToolsFramework/Source/ScriptableToolsFramework/Public/Tags/ScriptableToolGroupSet.h b/Engine/Plugins/Runtime/ScriptableToolsFramework/Source/ScriptableToolsFramework/Public/Tags/ScriptableToolGroupSet.h
--- a/Engine/Plugins/Runtime/ScriptableToolsFramework/Source/ScriptableToolsFramework/Public/Tags/ScriptableToolGroupSet.h
+++ b/Engine/Plugins/Runtime/ScriptableToolsFramework/Source/ScriptableToolsFramework/Public/Tags/ScriptableToolGroupSet.h
@@ -8,8 +8,10 @@
 
 #include "ScriptableToolGroupSet.generated.h"
 
+#define UE_API SCRIPTABLETOOLSFRAMEWORK_API
+
 USTRUCT()
-struct SCRIPTABLETOOLSFRAMEWORK_API FScriptableToolGroupSet
+struct FScriptableToolGroupSet
 {
 	GENERATED_BODY()
 
@@ -43,8 +45,10 @@ public:
 
 private:
 
-	void  SanitizeGroups();
+	UE_API void  SanitizeGroups();
 
 	UPROPERTY()
 	TSet<TSubclassOf<UScriptableToolGroupTag>> Groups;
-};
\ No newline at end of file
+};
+
+#undef UE_API
diff --git a/Engine/Plugins/Runtime/ScriptableToolsFramework/Source/ScriptableToolsFramework/Public/Tags/ScriptableToolGroupTag.h b/Engine/Plugins/Runtime/ScriptableToolsFramework/Source/ScriptableToolsFramework/Public/Tags/ScriptableToolGroupTag.h
--- a/Engine/Plugins/Runtime/ScriptableToolsFramework/Source/ScriptableToolsFramework/Public/Tags/ScriptableToolGroupTag.h
+++ b/Engine/Plugins/Runtime/ScriptableToolsFramework/Source/ScriptableToolsFramework/Public/Tags/ScriptableToolGroupTag.h
@@ -6,8 +6,8 @@
 
 #include "ScriptableToolGroupTag.generated.h"
 
-UCLASS(Abstract, Blueprintable, Const)
-class SCRIPTABLETOOLSFRAMEWORK_API UScriptableToolGroupTag
+UCLASS(MinimalAPI, Abstract, Blueprintable, Const)
+class UScriptableToolGroupTag
 	: public UObject
 {
 public:
diff --git a/Engine/Plugins/Runtime/ScriptableToolsFramework/Source/ScriptableToolsFramework/Public/Utility/ScriptableToolContextObjects.h b/Engine/Plugins/Runtime/ScriptableToolsFramework/Source/ScriptableToolsFramework/Public/Utility/ScriptableToolContextObjects.h
--- a/Engine/Plugins/Runtime/ScriptableToolsFramework/Source/ScriptableToolsFramework/Public/Utility/ScriptableToolContextObjects.h
+++ b/Engine/Plugins/Runtime/ScriptableToolsFramework/Source/ScriptableToolsFramework/Public/Utility/ScriptableToolContextObjects.h
@@ -12,8 +12,8 @@
 /**
  * Base class for context objects used in the Scriptable Tools Framework.
  */
-UCLASS()
-class SCRIPTABLETOOLSFRAMEWORK_API UScriptableToolContextObject : public UObject
+UCLASS(MinimalAPI)
+class UScriptableToolContextObject : public UObject
 {
 	GENERATED_BODY()
 public:
@@ -31,8 +31,8 @@ public:
 };
 
 
-UCLASS()
-class SCRIPTABLETOOLSFRAMEWORK_API UScriptableToolViewportWidgetAPI : public UScriptableToolContextObject
+UCLASS(MinimalAPI)
+class UScriptableToolViewportWidgetAPI : public UScriptableToolContextObject
 {
 	GENERATED_BODY()
 
@@ -82,4 +82,4 @@ private:
 	TUniqueFunction<void(TSharedRef<SWidget> InOverlaidWidget)> ClearOverlaidWidgetFunc;
 
 	TWeakPtr<SWidget> CurrentOverlaidWidget;
-};
\ No newline at end of file
+};
diff --git a/Engine/Plugins/Runtime/StateTree/Source/StateTreeEditorModule/Private/AsyncStateTreeDiff.h b/Engine/Plugins/Runtime/StateTree/Source/StateTreeEditorModule/Private/AsyncStateTreeDiff.h
--- a/Engine/Plugins/Runtime/StateTree/Source/StateTreeEditorModule/Private/AsyncStateTreeDiff.h
+++ b/Engine/Plugins/Runtime/StateTree/Source/StateTreeEditorModule/Private/AsyncStateTreeDiff.h
@@ -3,6 +3,8 @@
 
 #include "AsyncTreeDifferences.h"
 
+#define UE_API STATETREEEDITORMODULE_API
+
 class SStateTreeView;
 class UStateTreeState;
 
@@ -10,17 +12,17 @@ namespace UE::StateTree::Diff
 {
 struct FSingleDiffEntry;
 
-class STATETREEEDITORMODULE_API FAsyncDiff : public TAsyncTreeDifferences<TWeakObjectPtr<UStateTreeState>>
+class FAsyncDiff : public TAsyncTreeDifferences<TWeakObjectPtr<UStateTreeState>>
 {
 public:
-	FAsyncDiff(const TSharedRef<SStateTreeView>& LeftTree, const TSharedRef<SStateTreeView>& RightTree);
+	UE_API FAsyncDiff(const TSharedRef<SStateTreeView>& LeftTree, const TSharedRef<SStateTreeView>& RightTree);
 
-	void GetStateTreeDifferences(TArray<FSingleDiffEntry>& OutDiffEntries) const;
+	UE_API void GetStateTreeDifferences(TArray<FSingleDiffEntry>& OutDiffEntries) const;
 
 private:
-	void GetStatesDifferences(TArray<FSingleDiffEntry>& OutDiffEntries) const;
+	UE_API void GetStatesDifferences(TArray<FSingleDiffEntry>& OutDiffEntries) const;
 
-	static TAttribute<TArray<TWeakObjectPtr<UStateTreeState>>> RootNodesAttribute(TWeakPtr<SStateTreeView> StateTreeView);
+	static UE_API TAttribute<TArray<TWeakObjectPtr<UStateTreeState>>> RootNodesAttribute(TWeakPtr<SStateTreeView> StateTreeView);
 
 	TSharedPtr<SStateTreeView> LeftView;
 	TSharedPtr<SStateTreeView> RightView;
@@ -29,14 +31,14 @@ private:
 } // UE::StateTree::Diff
 
 template <>
-class STATETREEEDITORMODULE_API TTreeDiffSpecification<TWeakObjectPtr<UStateTreeState>>
+class TTreeDiffSpecification<TWeakObjectPtr<UStateTreeState>>
 {
 public:
-	bool AreValuesEqual(const TWeakObjectPtr<UStateTreeState>& StateTreeNodeA, const TWeakObjectPtr<UStateTreeState>& StateTreeNodeB, TArray<FPropertySoftPath>* OutDifferingProperties = nullptr) const;
+	UE_API bool AreValuesEqual(const TWeakObjectPtr<UStateTreeState>& StateTreeNodeA, const TWeakObjectPtr<UStateTreeState>& StateTreeNodeB, TArray<FPropertySoftPath>* OutDifferingProperties = nullptr) const;
 
-	bool AreMatching(const TWeakObjectPtr<UStateTreeState>& StateTreeNodeA, const TWeakObjectPtr<UStateTreeState>& StateTreeNodeB, TArray<FPropertySoftPath>* OutDifferingProperties = nullptr) const;
+	UE_API bool AreMatching(const TWeakObjectPtr<UStateTreeState>& StateTreeNodeA, const TWeakObjectPtr<UStateTreeState>& StateTreeNodeB, TArray<FPropertySoftPath>* OutDifferingProperties = nullptr) const;
 
-	void GetChildren(const TWeakObjectPtr<UStateTreeState>& InParent, TArray<TWeakObjectPtr<UStateTreeState>>& OutChildren) const;
+	UE_API void GetChildren(const TWeakObjectPtr<UStateTreeState>& InParent, TArray<TWeakObjectPtr<UStateTreeState>>& OutChildren) const;
 
 	bool ShouldMatchByValue(const TWeakObjectPtr<UStateTreeState>&) const
 	{
@@ -47,4 +49,6 @@ public:
 	{
 		return false;
 	}
-};
\ No newline at end of file
+};
+
+#undef UE_API
diff --git a/Engine/Plugins/Runtime/StateTree/Source/StateTreeEditorModule/Private/Customizations/StateTreeBlueprintPropertyRefDetails.h b/Engine/Plugins/Runtime/StateTree/Source/StateTreeEditorModule/Private/Customizations/StateTreeBlueprintPropertyRefDetails.h
--- a/Engine/Plugins/Runtime/StateTree/Source/StateTreeEditorModule/Private/Customizations/StateTreeBlueprintPropertyRefDetails.h
+++ b/Engine/Plugins/Runtime/StateTree/Source/StateTreeEditorModule/Private/Customizations/StateTreeBlueprintPropertyRefDetails.h
@@ -6,6 +6,8 @@
 #include "EdGraphSchema_K2.h"
 #include "StateTreeBlueprintPropertyRefDetails.generated.h"
 
+#define UE_API STATETREEEDITORMODULE_API
+
 class SWidget;
 class IPropertyHandle;
 
@@ -28,10 +30,12 @@ private:
 /**
  * Specific property ref schema to allow customizing the requirements (e.g. supported containers).
  */
-UCLASS()
-class STATETREEEDITORMODULE_API UStateTreePropertyRefSchema : public UEdGraphSchema_K2
+UCLASS(MinimalAPI)
+class UStateTreePropertyRefSchema : public UEdGraphSchema_K2
 {
 	GENERATED_BODY()
 public:
-	virtual bool SupportsPinTypeContainer(TWeakPtr<const FEdGraphSchemaAction> SchemaAction, const FEdGraphPinType& PinType, const EPinContainerType& ContainerType) const override;
+	UE_API virtual bool SupportsPinTypeContainer(TWeakPtr<const FEdGraphSchemaAction> SchemaAction, const FEdGraphPinType& PinType, const EPinContainerType& ContainerType) const override;
 };
+
+#undef UE_API
diff --git a/Engine/Plugins/Runtime/StateTree/Source/StateTreeEditorModule/Private/SStateTreeDiff.h b/Engine/Plugins/Runtime/StateTree/Source/StateTreeEditorModule/Private/SStateTreeDiff.h
--- a/Engine/Plugins/Runtime/StateTree/Source/StateTreeEditorModule/Private/SStateTreeDiff.h
+++ b/Engine/Plugins/Runtime/StateTree/Source/StateTreeEditorModule/Private/SStateTreeDiff.h
@@ -7,6 +7,8 @@
 #include "Subsystems/AssetEditorSubsystem.h"
 #include "Widgets/SCompoundWidget.h"
 
+#define UE_API STATETREEEDITORMODULE_API
+
 class FSpawnTabArgs;
 class FTabManager;
 class FUICommandList;
@@ -21,7 +23,7 @@ namespace UE::StateTree::Diff
 
 
 /** Panel used to display the state tree */
-struct STATETREEEDITORMODULE_API FDiffPanel
+struct FDiffPanel
 {
 	/** The asset that owns the state tree view we are showing */
 	TStrongObjectPtr<const UStateTree> StateTree = nullptr;
@@ -41,7 +43,7 @@ private:
 };
 
 /** Visual Diff between two StateTree assets */
-class STATETREEEDITORMODULE_API SDiffWidget : public SCompoundWidget
+class SDiffWidget : public SCompoundWidget
 {
 public:
 	SLATE_BEGIN_ARGS(SDiffWidget) {}
@@ -53,32 +55,32 @@ public:
 		SLATE_ARGUMENT(TSharedPtr<SWindow>, ParentWindow)
 	SLATE_END_ARGS()
 
-	SDiffWidget();
-	void Construct(const FArguments& InArgs);
-	virtual ~SDiffWidget();
+	UE_API SDiffWidget();
+	UE_API void Construct(const FArguments& InArgs);
+	UE_API virtual ~SDiffWidget();
 
 	/** Helper function to create a window that holds a diff widget */
-	static TSharedRef<SDiffWidget> CreateDiffWindow(FText WindowTitle, TNotNull<const UStateTree*> OldStateTree, TNotNull<const UStateTree*> NewStateTree, const FRevisionInfo& OldRevision, const FRevisionInfo& NewRevision);
+	static UE_API TSharedRef<SDiffWidget> CreateDiffWindow(FText WindowTitle, TNotNull<const UStateTree*> OldStateTree, TNotNull<const UStateTree*> NewStateTree, const FRevisionInfo& OldRevision, const FRevisionInfo& NewRevision);
 
 	/** Helper function to create a window that holds a diff widget (default window title) */
-	static TSharedRef<SDiffWidget> CreateDiffWindow(TNotNull<const UStateTree*> OldStateTree, TNotNull<const UStateTree*> NewStateTree, const FRevisionInfo& OldRevision, const FRevisionInfo& NewRevision, const UClass* StateTreeClass);
+	static UE_API TSharedRef<SDiffWidget> CreateDiffWindow(TNotNull<const UStateTree*> OldStateTree, TNotNull<const UStateTree*> NewStateTree, const FRevisionInfo& OldRevision, const FRevisionInfo& NewRevision, const UClass* StateTreeClass);
 
 protected:
 	/** Called when user clicks button to go to next difference */
-	void NextDiff() const;
+	UE_API void NextDiff() const;
 
 	/** Called when user clicks button to go to prev difference */
-	void PrevDiff() const;
+	UE_API void PrevDiff() const;
 
 	/** Called to determine whether we have a list of differences to cycle through */
-	bool HasNextDiff() const;
-	bool HasPrevDiff() const;
+	UE_API bool HasNextDiff() const;
+	UE_API bool HasPrevDiff() const;
 
 	/** Function used to generate the list of differences and the widgets needed to calculate that list */
-	void GenerateDifferencesList();
+	UE_API void GenerateDifferencesList();
 
 	/** Called when editor may need to be closed */
-	void HandleAssetEditorRequestClose(UObject* Asset, EAssetEditorCloseReason CloseReason);
+	UE_API void HandleAssetEditorRequestClose(UObject* Asset, EAssetEditorCloseReason CloseReason);
 
 	struct FStateTreePanel
 	{
@@ -86,26 +88,26 @@ protected:
 		TSharedPtr<FDiffControl> DiffControl;
 	};
 
-	void GenerateDiffPanel();
+	UE_API void GenerateDiffPanel();
 
-	void HandleStateDiffEntryFocused(const FSingleDiffEntry& StateDiff);
+	UE_API void HandleStateDiffEntryFocused(const FSingleDiffEntry& StateDiff);
 	
-	void SetDetailsDiff(const UStateTreeState* OldState = nullptr, const UStateTreeState* NewState = nullptr);
+	UE_API void SetDetailsDiff(const UStateTreeState* OldState = nullptr, const UStateTreeState* NewState = nullptr);
 
-	static void AddStateTreeExtensionToDetailsView(const TSharedRef<IDetailsView>& DetailsView);
+	static UE_API void AddStateTreeExtensionToDetailsView(const TSharedRef<IDetailsView>& DetailsView);
 
-	void AddBindingDiffToDiffEntries(TArray<FSingleObjectDiffEntry>& OutEntries);
+	UE_API void AddBindingDiffToDiffEntries(TArray<FSingleObjectDiffEntry>& OutEntries);
 
-	static TSharedRef<SWidget> GenerateCustomDiffEntryWidget(const FSingleObjectDiffEntry& DiffEntry, FText& ObjectName, const UStateTreeState* OldState, const UStateTreeState* NewState);
+	static UE_API TSharedRef<SWidget> GenerateCustomDiffEntryWidget(const FSingleObjectDiffEntry& DiffEntry, FText& ObjectName, const UStateTreeState* OldState, const UStateTreeState* NewState);
 
-	static void OrganizeDiffEntries(TArray<TSharedPtr<FBlueprintDifferenceTreeEntry>>& OutDiffTreeEntries, const TArray<FSingleObjectDiffEntry>& DiffEntries,
+	static UE_API void OrganizeDiffEntries(TArray<TSharedPtr<FBlueprintDifferenceTreeEntry>>& OutDiffTreeEntries, const TArray<FSingleObjectDiffEntry>& DiffEntries,
 		TFunctionRef<TSharedPtr<FBlueprintDifferenceTreeEntry>(const FSingleObjectDiffEntry&)> GenerateDiffTreeEntry,
 		TFunctionRef<TSharedPtr<FBlueprintDifferenceTreeEntry>(FText&)> GenerateCategoryEntry,
 		const UStateTreeState* OldState, const UStateTreeState* NewState);
 
-	bool ShouldHighlightRow(const TUniquePtr<FAsyncDetailViewDiff::DiffNodeType>& DiffNode);
+	UE_API bool ShouldHighlightRow(const TUniquePtr<FAsyncDetailViewDiff::DiffNodeType>& DiffNode);
 
-	static FLinearColor GetRowHighlightColor(const TUniquePtr<FAsyncDetailViewDiff::DiffNodeType>& DiffNode);
+	static UE_API FLinearColor GetRowHighlightColor(const TUniquePtr<FAsyncDetailViewDiff::DiffNodeType>& DiffNode);
 
 	/** The panel used to show the old revision */
 	FDiffPanel OldAssetPanel;
@@ -142,3 +144,5 @@ private:
 };
 
 } // UE::StateTree::Diff
+
+#undef UE_API
diff --git a/Engine/Plugins/Runtime/StateTree/Source/StateTreeEditorModule/Private/SStateTreeSplitter.h b/Engine/Plugins/Runtime/StateTree/Source/StateTreeEditorModule/Private/SStateTreeSplitter.h
--- a/Engine/Plugins/Runtime/StateTree/Source/StateTreeEditorModule/Private/SStateTreeSplitter.h
+++ b/Engine/Plugins/Runtime/StateTree/Source/StateTreeEditorModule/Private/SStateTreeSplitter.h
@@ -4,6 +4,8 @@
 #include "StateTreeDiffHelper.h"
 #include "Widgets/SCompoundWidget.h"
 
+#define UE_API STATETREEEDITORMODULE_API
+
 class SStateTreeView;
 
 namespace UE::StateTree::Diff
@@ -11,7 +13,7 @@ namespace UE::StateTree::Diff
 class FAsyncDiff;
 
 /** Splitter that allows you to provide an FAsyncStateTreeViewDiff to connect like-properties between two or more state tree panels */
-class STATETREEEDITORMODULE_API SDiffSplitter : public SCompoundWidget
+class SDiffSplitter : public SCompoundWidget
 {
 public:
 	class FSlot : public TSlotBase<FSlot>
@@ -28,7 +30,7 @@ public:
 			SLATE_ATTRIBUTE(TSharedPtr<FAsyncDiff>, DifferencesWithRightPanel)
 		SLATE_SLOT_END_ARGS()
 	};
-	static FSlot::FSlotArguments Slot();
+	static UE_API FSlot::FSlotArguments Slot();
 
 	SLATE_BEGIN_ARGS(SDiffSplitter)
 	{
@@ -36,14 +38,14 @@ public:
 		SLATE_SLOT_ARGUMENT(FSlot, Slots)
 	SLATE_END_ARGS()
 
-	void Construct(const FArguments& InArgs);
+	UE_API void Construct(const FArguments& InArgs);
 
-	void AddSlot(const FSlot::FSlotArguments& SlotArgs, int32 Index = INDEX_NONE);
+	UE_API void AddSlot(const FSlot::FSlotArguments& SlotArgs, int32 Index = INDEX_NONE);
 
-	void HandleSelectionChanged(const FStateSoftPath& StatePath, const FStateSoftPath& SecondaryStatePath);
+	UE_API void HandleSelectionChanged(const FStateSoftPath& StatePath, const FStateSoftPath& SecondaryStatePath);
 
 protected:
-	void HandleSelectionChanged(const TArray<TWeakObjectPtr<UStateTreeState>>& SelectedStates);
+	UE_API void HandleSelectionChanged(const TArray<TWeakObjectPtr<UStateTreeState>>& SelectedStates);
 
 private:
 	struct FPanel
@@ -68,4 +70,6 @@ private:
 
 	FStateSoftPath SelectedState;
 };
-} // UE::StateTree::Diff
\ No newline at end of file
+} // UE::StateTree::Diff
+
+#undef UE_API
diff --git a/Engine/Plugins/Runtime/StateTree/Source/StateTreeEditorModule/Private/StateTreeDiffControl.h b/Engine/Plugins/Runtime/StateTree/Source/StateTreeEditorModule/Private/StateTreeDiffControl.h
--- a/Engine/Plugins/Runtime/StateTree/Source/StateTreeEditorModule/Private/StateTreeDiffControl.h
+++ b/Engine/Plugins/Runtime/StateTree/Source/StateTreeEditorModule/Private/StateTreeDiffControl.h
@@ -9,6 +9,8 @@
 #include "UObject/ObjectKey.h"
 #include "UObject/StrongObjectPtr.h"
 
+#define UE_API STATETREEEDITORMODULE_API
+
 class FBlueprintDifferenceTreeEntry;
 class FStateTreeViewModel;
 class FUICommandList;
@@ -21,13 +23,13 @@ namespace UE::StateTree::Diff
 {
 struct FSingleDiffEntry;
 
-class STATETREEEDITORMODULE_API FDiffWidgets
+class FDiffWidgets
 {
 public:
-	explicit FDiffWidgets(const UStateTree* InStateTree);
+	UE_API explicit FDiffWidgets(const UStateTree* InStateTree);
 
 	/** Returns actual widget that is used to display trees */
-	TSharedRef<SStateTreeView> GetStateTreeWidget() const;
+	UE_API TSharedRef<SStateTreeView> GetStateTreeWidget() const;
 
 private:
 	TSharedPtr<SStateTreeView> StateTreeTreeView;
@@ -36,35 +38,35 @@ private:
 
 DECLARE_MULTICAST_DELEGATE_OneParam(FOnStateDiffEntryFocused, const FSingleDiffEntry&)
 
-class STATETREEEDITORMODULE_API FDiffControl : public TSharedFromThis<FDiffControl>
+class FDiffControl : public TSharedFromThis<FDiffControl>
 {
 public:
 	FDiffControl() = delete;
 	FDiffControl(const FDiffControl& Other) = delete;
 	FDiffControl(const FDiffControl&& Other) = delete;
-	FDiffControl(const UStateTree* InOldObject, const UStateTree* InNewObject, const FOnDiffEntryFocused& InSelectionCallback);
-	~FDiffControl();
+	UE_API FDiffControl(const UStateTree* InOldObject, const UStateTree* InNewObject, const FOnDiffEntryFocused& InSelectionCallback);
+	UE_API ~FDiffControl();
 
-	TSharedRef<SStateTreeView> GetDetailsWidget(const UStateTree* Object) const;
+	UE_API TSharedRef<SStateTreeView> GetDetailsWidget(const UStateTree* Object) const;
 
 	FOnStateDiffEntryFocused& GetOnStateDiffEntryFocused()
 	{
 		return OnStateDiffEntryFocused;
 	};
 
-	void GenerateTreeEntries(TArray<TSharedPtr<FBlueprintDifferenceTreeEntry>>& OutDifferences);
+	UE_API void GenerateTreeEntries(TArray<TSharedPtr<FBlueprintDifferenceTreeEntry>>& OutDifferences);
 	TConstArrayView<FSingleDiffEntry> GetBindingDifferences() const
 	{
 		return BindingDiffs;
 	}
 
 protected:
-	static FText RightRevision;
-	static TSharedRef<SWidget> GenerateSingleEntryWidget(FSingleDiffEntry DiffEntry, FText ObjectName);
+	static UE_API FText RightRevision;
+	static UE_API TSharedRef<SWidget> GenerateSingleEntryWidget(FSingleDiffEntry DiffEntry, FText ObjectName);
 
-	TSharedRef<SStateTreeView> InsertObject(TNotNull<const UStateTree*> StateTree);
+	UE_API TSharedRef<SStateTreeView> InsertObject(TNotNull<const UStateTree*> StateTree);
 
-	void OnSelectDiffEntry(const FSingleDiffEntry StateDiff);
+	UE_API void OnSelectDiffEntry(const FSingleDiffEntry StateDiff);
 
 	FOnDiffEntryFocused OnDiffEntryFocused;
 	FOnStateDiffEntryFocused OnStateDiffEntryFocused;
@@ -81,4 +83,6 @@ protected:
 	TMap<FObjectKey, FDiffWidgets> StateTreeDiffWidgets;
 };
 
-} // UE::StateTree::Diff
\ No newline at end of file
+} // UE::StateTree::Diff
+
+#undef UE_API
diff --git a/Engine/Plugins/Runtime/StateTree/Source/StateTreeEditorModule/Public/IStateTreeEditor.h b/Engine/Plugins/Runtime/StateTree/Source/StateTreeEditorModule/Public/IStateTreeEditor.h
--- a/Engine/Plugins/Runtime/StateTree/Source/StateTreeEditorModule/Public/IStateTreeEditor.h
+++ b/Engine/Plugins/Runtime/StateTree/Source/StateTreeEditorModule/Public/IStateTreeEditor.h
@@ -6,7 +6,7 @@
 
 
 /** StateTree Editor public interface */
-class STATETREEEDITORMODULE_API IStateTreeEditor : public FAssetEditorToolkit
+class IStateTreeEditor : public FAssetEditorToolkit
 {
 public:
 
diff --git a/Engine/Plugins/Runtime/StateTree/Source/StateTreeEditorModule/Public/IStateTreeEditorHost.h b/Engine/Plugins/Runtime/StateTree/Source/StateTreeEditorModule/Public/IStateTreeEditorHost.h
--- a/Engine/Plugins/Runtime/StateTree/Source/StateTreeEditorModule/Public/IStateTreeEditorHost.h
+++ b/Engine/Plugins/Runtime/StateTree/Source/StateTreeEditorModule/Public/IStateTreeEditorHost.h
@@ -36,8 +36,8 @@ public:
 	virtual TSharedPtr<UE::StateTreeEditor::FWorkspaceTabHost> GetTabHost() const = 0;
 };
 
-UCLASS()
-class STATETREEEDITORMODULE_API UStateTreeEditorContext : public UObject
+UCLASS(MinimalAPI)
+class UStateTreeEditorContext : public UObject
 {
 	GENERATED_BODY()
 public:
diff --git a/Engine/Plugins/Runtime/StateTree/Source/StateTreeEditorModule/Public/StateTreeCompiler.h b/Engine/Plugins/Runtime/StateTree/Source/StateTreeEditorModule/Public/StateTreeCompiler.h
--- a/Engine/Plugins/Runtime/StateTree/Source/StateTreeEditorModule/Public/StateTreeCompiler.h
+++ b/Engine/Plugins/Runtime/StateTree/Source/StateTreeEditorModule/Public/StateTreeCompiler.h
@@ -4,6 +4,8 @@
 
 #include "StateTreePropertyBindingCompiler.h"
 
+#define UE_API STATETREEEDITORMODULE_API
+
 struct FStructView;
 
 enum class EStateTreeExpressionOperand : uint8;
@@ -23,7 +25,7 @@ struct FStateTreeNodeBase;
  * Helper class to convert StateTree editor representation into a compact data.
  * Holds data needed during compiling.
  */
-struct STATETREEEDITORMODULE_API FStateTreeCompiler
+struct FStateTreeCompiler
 {
 public:
 
@@ -32,40 +34,40 @@ public:
 	{
 	}
 	
-	bool Compile(UStateTree& InStateTree);
+	UE_API bool Compile(UStateTree& InStateTree);
 	
 private:
 
 	/** Resolves the state a transition points to, and the optional fallback for failing to enter the state. SourceState is nullptr for global tasks. */
-	bool ResolveTransitionStateAndFallback(const UStateTreeState* SourceState, const FStateTreeStateLink& Link, FStateTreeStateHandle& OutTransitionHandle, EStateTreeSelectionFallback& OutFallback) const;
-	FStateTreeStateHandle GetStateHandle(const FGuid& StateID) const;
-	UStateTreeState* GetState(const FGuid& StateID) const;
-
-	bool CreateParameters();
-	bool CreateStates();
-	bool CreateStateRecursive(UStateTreeState& State, const FStateTreeStateHandle Parent);
-	bool CreateEvaluators();
-	bool CreateGlobalTasks();
-	bool CreateStateTasksAndParameters();
-	bool CreateStateTransitions();
-	bool CreateStateConsiderations();
+	UE_API bool ResolveTransitionStateAndFallback(const UStateTreeState* SourceState, const FStateTreeStateLink& Link, FStateTreeStateHandle& OutTransitionHandle, EStateTreeSelectionFallback& OutFallback) const;
+	UE_API FStateTreeStateHandle GetStateHandle(const FGuid& StateID) const;
+	UE_API UStateTreeState* GetState(const FGuid& StateID) const;
+
+	UE_API bool CreateParameters();
+	UE_API bool CreateStates();
+	UE_API bool CreateStateRecursive(UStateTreeState& State, const FStateTreeStateHandle Parent);
+	UE_API bool CreateEvaluators();
+	UE_API bool CreateGlobalTasks();
+	UE_API bool CreateStateTasksAndParameters();
+	UE_API bool CreateStateTransitions();
+	UE_API bool CreateStateConsiderations();
 	
-	bool CreateBindingsForNodes(TConstArrayView<FStateTreeEditorNode> EditorNodes, FStateTreeIndex16 NodesBegin, TArray<FInstancedStruct>& Instances);
-	bool CreateBindingsForStruct(const FStateTreeBindableStructDesc& TargetStruct, FStateTreeDataView TargetValue, FStateTreeIndex16 PropertyFuncsBegin, FStateTreeIndex16 PropertyFuncsEnd, FStateTreeIndex16& OutBatchIndex);
+	UE_API bool CreateBindingsForNodes(TConstArrayView<FStateTreeEditorNode> EditorNodes, FStateTreeIndex16 NodesBegin, TArray<FInstancedStruct>& Instances);
+	UE_API bool CreateBindingsForStruct(const FStateTreeBindableStructDesc& TargetStruct, FStateTreeDataView TargetValue, FStateTreeIndex16 PropertyFuncsBegin, FStateTreeIndex16 PropertyFuncsEnd, FStateTreeIndex16& OutBatchIndex);
 
-	bool CreatePropertyFunctionsForStruct(FGuid StructID);
-	bool CreatePropertyFunction(const FStateTreeEditorNode& FuncEditorNode);
+	UE_API bool CreatePropertyFunctionsForStruct(FGuid StructID);
+	UE_API bool CreatePropertyFunction(const FStateTreeEditorNode& FuncEditorNode);
 	
-	bool CreateConditions(UStateTreeState& State, const FString& StatePath, TConstArrayView<FStateTreeEditorNode> Conditions);
-	bool CreateCondition(UStateTreeState& State, const FString& StatePath, const FStateTreeEditorNode& CondNode, const EStateTreeExpressionOperand Operand, const int8 DeltaIndent);
-	bool CreateConsiderations(UStateTreeState& State, const FString& StatePath, TConstArrayView<FStateTreeEditorNode> Considerations);
-	bool CreateConsideration(UStateTreeState& State, const FString& StatePath, const FStateTreeEditorNode& ConsiderationNode, const EStateTreeExpressionOperand Operand, const int8 DeltaIndent);
+	UE_API bool CreateConditions(UStateTreeState& State, const FString& StatePath, TConstArrayView<FStateTreeEditorNode> Conditions);
+	UE_API bool CreateCondition(UStateTreeState& State, const FString& StatePath, const FStateTreeEditorNode& CondNode, const EStateTreeExpressionOperand Operand, const int8 DeltaIndent);
+	UE_API bool CreateConsiderations(UStateTreeState& State, const FString& StatePath, TConstArrayView<FStateTreeEditorNode> Considerations);
+	UE_API bool CreateConsideration(UStateTreeState& State, const FString& StatePath, const FStateTreeEditorNode& ConsiderationNode, const EStateTreeExpressionOperand Operand, const int8 DeltaIndent);
 
 	template<class T = FStateTreeNodeBase>
 	T* CreateNodeWithSharedInstanceData(UStateTreeState* State, const FStateTreeEditorNode& EditorNode, FStateTreeBindableStructDesc& StructDesc);
 
-	bool CreateTask(UStateTreeState* State, const FStateTreeEditorNode& TaskNode, const FStateTreeDataHandle TaskDataHandle);
-	bool CreateEvaluator(const FStateTreeEditorNode& EvalNode, const FStateTreeDataHandle EvalDataHandle);
+	UE_API bool CreateTask(UStateTreeState* State, const FStateTreeEditorNode& TaskNode, const FStateTreeDataHandle TaskDataHandle);
+	UE_API bool CreateEvaluator(const FStateTreeEditorNode& EvalNode, const FStateTreeDataHandle EvalDataHandle);
 
 	struct FValidatedPathBindings
 	{
@@ -75,13 +77,13 @@ private:
 		TArray<FStateTreePropertyPathBinding> ReferenceBindings;
 	};
 
-	bool GetAndValidateBindings(const FStateTreeBindableStructDesc& TargetStruct, FStateTreeDataView TargetValue, FValidatedPathBindings& OutValidatedBindings) const;
-	bool IsPropertyOfTypeOrChild(UScriptStruct& Type, const FStateTreeBindableStructDesc& Struct, FPropertyBindingPath Path) const;
-	bool ValidateStructRef(const FStateTreeBindableStructDesc& SourceStruct, FPropertyBindingPath SourcePath,
+	UE_API bool GetAndValidateBindings(const FStateTreeBindableStructDesc& TargetStruct, FStateTreeDataView TargetValue, FValidatedPathBindings& OutValidatedBindings) const;
+	UE_API bool IsPropertyOfTypeOrChild(UScriptStruct& Type, const FStateTreeBindableStructDesc& Struct, FPropertyBindingPath Path) const;
+	UE_API bool ValidateStructRef(const FStateTreeBindableStructDesc& SourceStruct, FPropertyBindingPath SourcePath,
 							const FStateTreeBindableStructDesc& TargetStruct, FPropertyBindingPath TargetPath) const;
-	bool CompileAndValidateNode(const UStateTreeState* SourceState, const FStateTreeBindableStructDesc& NodeDesc, FStructView NodeView, const FStateTreeDataView InstanceData);
+	UE_API bool CompileAndValidateNode(const UStateTreeState* SourceState, const FStateTreeBindableStructDesc& NodeDesc, FStructView NodeView, const FStateTreeDataView InstanceData);
 
-	void InstantiateStructSubobjects(FStructView Struct);
+	UE_API void InstantiateStructSubobjects(FStructView Struct);
 	
 	FStateTreeCompilerLog& Log;
 	UStateTree* StateTree = nullptr;
@@ -182,3 +184,5 @@ namespace UE::StateTree::Compiler
 	const UScriptStruct* GetBaseStructFromMetaData(const FProperty* Property, FString& OutBaseStructName);
 
 }; // UE::StateTree::Compiler
+
+#undef UE_API
diff --git a/Engine/Plugins/Runtime/StateTree/Source/StateTreeEditorModule/Public/StateTreeCompilerLog.h b/Engine/Plugins/Runtime/StateTree/Source/StateTreeEditorModule/Public/StateTreeCompilerLog.h
--- a/Engine/Plugins/Runtime/StateTree/Source/StateTreeEditorModule/Public/StateTreeCompilerLog.h
+++ b/Engine/Plugins/Runtime/StateTree/Source/StateTreeEditorModule/Public/StateTreeCompilerLog.h
@@ -6,11 +6,13 @@
 #include "StateTreePropertyBindings.h"
 #include "StateTreeCompilerLog.generated.h"
 
+#define UE_API STATETREEEDITORMODULE_API
+
 class UStateTreeState;
 
 /** StateTree compiler log message */
 USTRUCT()
-struct STATETREEEDITORMODULE_API FStateTreeCompilerLogMessage
+struct FStateTreeCompilerLogMessage
 {
 	GENERATED_BODY()
 
@@ -42,7 +44,7 @@ struct STATETREEEDITORMODULE_API FStateTreeCompilerLogMessage
 
 /** Message log for StateTree compilation */
 USTRUCT()
-struct STATETREEEDITORMODULE_API FStateTreeCompilerLog
+struct FStateTreeCompilerLog
 {
 	GENERATED_BODY()
 
@@ -92,13 +94,13 @@ struct STATETREEEDITORMODULE_API FStateTreeCompilerLog
 	}
 
 	/** Generates tokenized messages from the messages in this log */
-	TArray<TSharedRef<FTokenizedMessage>> ToTokenizedMessages() const;
+	UE_API TArray<TSharedRef<FTokenizedMessage>> ToTokenizedMessages() const;
 
 	/** Appends StateTree log to log listing. */
-	void AppendToLog(class IMessageLogListing* LogListing) const;
+	UE_API void AppendToLog(class IMessageLogListing* LogListing) const;
 
 	/** Dumps StateTree log to log */
-	void DumpToLog(const FLogCategoryBase& Category) const;
+	UE_API void DumpToLog(const FLogCategoryBase& Category) const;
 	
 protected:
 	UPROPERTY()
@@ -126,3 +128,5 @@ struct FStateTreeCompilerLogStateScope
 	FStateTreeCompilerLog& Log; 
 	const UStateTreeState* State = nullptr;
 };
+
+#undef UE_API
diff --git a/Engine/Plugins/Runtime/StateTree/Source/StateTreeEditorModule/Public/StateTreeEditingSubsystem.h b/Engine/Plugins/Runtime/StateTree/Source/StateTreeEditorModule/Public/StateTreeEditingSubsystem.h
--- a/Engine/Plugins/Runtime/StateTree/Source/StateTreeEditorModule/Public/StateTreeEditingSubsystem.h
+++ b/Engine/Plugins/Runtime/StateTree/Source/StateTreeEditorModule/Public/StateTreeEditingSubsystem.h
@@ -9,40 +9,44 @@
 
 #include "StateTreeEditingSubsystem.generated.h"
 
+#define UE_API STATETREEEDITORMODULE_API
+
 class SWidget;
 class FStateTreeViewModel;
 class FUICommandList;
 struct FStateTreeCompilerLog;
 
-UCLASS()
-class STATETREEEDITORMODULE_API UStateTreeEditingSubsystem : public UEditorSubsystem
+UCLASS(MinimalAPI)
+class UStateTreeEditingSubsystem : public UEditorSubsystem
 {
 	GENERATED_BODY()
 public:
-	UStateTreeEditingSubsystem();
-	virtual void BeginDestroy() override;
+	UE_API UStateTreeEditingSubsystem();
+	UE_API virtual void BeginDestroy() override;
 	
-	TSharedRef<FStateTreeViewModel> FindOrAddViewModel(TNotNull<UStateTree*> InStateTree);
+	UE_API TSharedRef<FStateTreeViewModel> FindOrAddViewModel(TNotNull<UStateTree*> InStateTree);
 	
-	static bool CompileStateTree(TNotNull<UStateTree*> InStateTree,  FStateTreeCompilerLog& InOutLog);
+	static UE_API bool CompileStateTree(TNotNull<UStateTree*> InStateTree,  FStateTreeCompilerLog& InOutLog);
 	
 	/** Create a StateTreeView widget for the viewmodel. */
-	static TSharedRef<SWidget> GetStateTreeView(TSharedRef<FStateTreeViewModel> InViewModel, const TSharedRef<FUICommandList>& TreeViewCommandList);
+	static UE_API TSharedRef<SWidget> GetStateTreeView(TSharedRef<FStateTreeViewModel> InViewModel, const TSharedRef<FUICommandList>& TreeViewCommandList);
 	
 	/**
 	 * Validates and applies the schema restrictions on the StateTree.
 	 * Also serves as the "safety net" of fixing up editor data following an editor operation.
 	 * Updates state's link, removes the unused node while validating the StateTree asset.
 	 */
-	static void ValidateStateTree(TNotNull<UStateTree*> InStateTree);
+	static UE_API void ValidateStateTree(TNotNull<UStateTree*> InStateTree);
 
 	/** Calculates editor data hash of the asset. */
-	static uint32 CalculateStateTreeHash(TNotNull<const UStateTree*> InStateTree);
+	static UE_API uint32 CalculateStateTreeHash(TNotNull<const UStateTree*> InStateTree);
 	
 private:
-	void HandlePostGarbageCollect();
+	UE_API void HandlePostGarbageCollect();
 
 protected:
 	TMap<FObjectKey, TSharedPtr<FStateTreeViewModel>> StateTreeViewModels;
 	FDelegateHandle PostGarbageCollectHandle;
 };
+
+#undef UE_API
