ScriptableTools: Fix issue with marking properties read only making them invisible. Also provides dynamic updates to the details panel when the calls happen outside setup.

diff --git a/Engine/Plugins/Runtime/ScriptableToolsFramework/Source/ScriptableToolsFramework/Private/ScriptableInteractiveTool.cpp b/Engine/Plugins/Runtime/ScriptableToolsFramework/Source/ScriptableToolsFramework/Private/ScriptableInteractiveTool.cpp
--- a/Engine/Plugins/Runtime/ScriptableToolsFramework/Source/ScriptableToolsFramework/Private/ScriptableInteractiveTool.cpp
+++ b/Engine/Plugins/Runtime/ScriptableToolsFramework/Source/ScriptableToolsFramework/Private/ScriptableInteractiveTool.cpp
@@ -89,14 +89,17 @@ UScriptableInteractiveToolPropertySet* UScriptableInteractiveToolPropertySet::Se
 	if (!bReadOnly)
 	{
 		Property->ClearPropertyFlags(EPropertyFlags::CPF_EditConst);
-		Property->SetPropertyFlags(EPropertyFlags::CPF_Edit);
 	}
 	else
 	{
-		Property->ClearPropertyFlags(EPropertyFlags::CPF_Edit);
 		Property->SetPropertyFlags(EPropertyFlags::CPF_EditConst);
 	}
 
+	if (UScriptableInteractiveTool* OwningTool = ParentTool.Get())
+	{
+		OwningTool->OnPropertyModifiedDirectlyByTool.Broadcast(this);
+	}
+
 	return this;
 }
 
@@ -110,7 +113,6 @@ UScriptableInteractiveToolPropertySet* UScriptableInteractiveToolPropertySet::Se
 	}
 	if (bHidden)
 	{
-		Property->ClearPropertyFlags(EPropertyFlags::CPF_EditConst);
 		Property->ClearPropertyFlags(EPropertyFlags::CPF_Edit);
 	}
 	else
@@ -118,6 +120,11 @@ UScriptableInteractiveToolPropertySet* UScriptableInteractiveToolPropertySet::Se
 		Property->SetPropertyFlags(EPropertyFlags::CPF_Edit);
 	}
 
+	if (UScriptableInteractiveTool* OwningTool = ParentTool.Get())
+	{
+		OwningTool->OnPropertyModifiedDirectlyByTool.Broadcast(this);
+	}
+
 	return this;
 }
 
