ScriptableTools: Fixed reachable garbage objects warnings when compiling ScriptableTools containing ScriptablePoint/Line/Triangle sets.

diff --git a/Engine/Plugins/Runtime/ScriptableToolsFramework/Source/ScriptableToolsFramework/Private/Drawing/ScriptableToolLineSet.cpp b/Engine/Plugins/Runtime/ScriptableToolsFramework/Source/ScriptableToolsFramework/Private/Drawing/ScriptableToolLineSet.cpp
--- a/Engine/Plugins/Runtime/ScriptableToolsFramework/Source/ScriptableToolsFramework/Private/Drawing/ScriptableToolLineSet.cpp
+++ b/Engine/Plugins/Runtime/ScriptableToolsFramework/Source/ScriptableToolsFramework/Private/Drawing/ScriptableToolLineSet.cpp
@@ -9,55 +9,74 @@
 void UScriptableToolLineSet::Initialize(TObjectPtr<UPreviewGeometry> PreviewGeometry)
 {
 	FString LineSetID = FGuid::NewGuid().ToString();
-	LineSet = PreviewGeometry->AddLineSet(LineSetID);
+	WeakLineSet = PreviewGeometry->AddLineSet(LineSetID);
 }
 
 void UScriptableToolLineSet::OnTick()
 {
-	for (TObjectPtr<UScriptableToolLine> LineComponent : LineComponents)
+	if (ULineSetComponent* ResolvedLineSet = WeakLineSet.Get())
 	{
-		if (LineComponent->IsDirty())
+		for (TObjectPtr<UScriptableToolLine> LineComponent : LineComponents)
 		{
-			int32 LineID = LineComponent->GetLineID();
-			FRenderableLine LineDescription = LineComponent->GenerateLineDescription();
+			if (LineComponent->IsDirty())
+			{
+				int32 LineID = LineComponent->GetLineID();
+				FRenderableLine LineDescription = LineComponent->GenerateLineDescription();
 
-			LineSet->SetLineStart(LineID, LineDescription.Start);
-			LineSet->SetLineEnd(LineID, LineDescription.End);
-			LineSet->SetLineColor(LineID, LineDescription.Color);
-			LineSet->SetLineThickness(LineID, LineDescription.Thickness);
+				ResolvedLineSet->SetLineStart(LineID, LineDescription.Start);
+				ResolvedLineSet->SetLineEnd(LineID, LineDescription.End);
+				ResolvedLineSet->SetLineColor(LineID, LineDescription.Color);
+				ResolvedLineSet->SetLineThickness(LineID, LineDescription.Thickness);
+			}
 		}
 	}
 }
 
 UScriptableToolLine* UScriptableToolLineSet::AddLine()
 {
-	LineComponents.Add(NewObject<UScriptableToolLine>(this));
-	UScriptableToolLine* NewLine = LineComponents.Last();
-	NewLine->SetLineID(LineSet->AddLine(NewLine->GenerateLineDescription()));
+	UScriptableToolLine* NewLine = nullptr;
+	if (ULineSetComponent* ResolvedLineSet = WeakLineSet.Get())
+	{
+		LineComponents.Add(NewObject<UScriptableToolLine>(this));
+		NewLine = LineComponents.Last();
+		NewLine->SetLineID(ResolvedLineSet->AddLine(NewLine->GenerateLineDescription()));
+	}
 	return NewLine;
 }
 
 void UScriptableToolLineSet::RemoveLine(UScriptableToolLine* Line)
 {
-	if (ensure(Line))
+	if (ULineSetComponent* ResolvedLineSet = WeakLineSet.Get())
 	{
-		LineSet->RemoveLine(Line->GetLineID());
-		LineComponents.Remove(Line);
+		if (ensure(Line))
+		{
+			ResolvedLineSet->RemoveLine(Line->GetLineID());
+			LineComponents.Remove(Line);
+		}
 	}
 }
 
 void UScriptableToolLineSet::RemoveAllLines()
 {
-	LineSet->Clear();
+	if (ULineSetComponent* ResolvedLineSet = WeakLineSet.Get())
+	{
+		ResolvedLineSet->Clear();
+	}
 	LineComponents.Empty();
 }
 
 void UScriptableToolLineSet::SetAllLinesColor(FColor Color)
 {
-	LineSet->SetAllLinesColor(Color);
+	if (ULineSetComponent* ResolvedLineSet = WeakLineSet.Get())
+	{
+		ResolvedLineSet->SetAllLinesColor(Color);
+	}
 }
 
 void UScriptableToolLineSet::SetAllLinesThickness(float Thickness)
 {
-	LineSet->SetAllLinesThickness(Thickness);
+	if (ULineSetComponent* ResolvedLineSet = WeakLineSet.Get())
+	{
+		ResolvedLineSet->SetAllLinesThickness(Thickness);
+	}
 }
\ No newline at end of file
diff --git a/Engine/Plugins/Runtime/ScriptableToolsFramework/Source/ScriptableToolsFramework/Private/Drawing/ScriptableToolPointSet.cpp b/Engine/Plugins/Runtime/ScriptableToolsFramework/Source/ScriptableToolsFramework/Private/Drawing/ScriptableToolPointSet.cpp
--- a/Engine/Plugins/Runtime/ScriptableToolsFramework/Source/ScriptableToolsFramework/Private/Drawing/ScriptableToolPointSet.cpp
+++ b/Engine/Plugins/Runtime/ScriptableToolsFramework/Source/ScriptableToolsFramework/Private/Drawing/ScriptableToolPointSet.cpp
@@ -9,54 +9,73 @@
 void UScriptableToolPointSet::Initialize(TObjectPtr<UPreviewGeometry> PreviewGeometry)
 {
 	FString PointSetID = FGuid::NewGuid().ToString();
-	PointSet = PreviewGeometry->AddPointSet(PointSetID);
+	WeakPointSet = PreviewGeometry->AddPointSet(PointSetID);
 }
 
 void UScriptableToolPointSet::OnTick()
 {
-	for (TObjectPtr<UScriptableToolPoint> PointComponent : PointComponents)
+	if (UPointSetComponent* ResolvedPointSet = WeakPointSet.Get())
 	{
-		if (PointComponent->IsDirty())
+		for (TObjectPtr<UScriptableToolPoint> PointComponent : PointComponents)
 		{
-			int32 PointID = PointComponent->GetPointID();
-			FRenderablePoint PointDescription = PointComponent->GeneratePointDescription();
+			if (PointComponent->IsDirty())
+			{
+				int32 PointID = PointComponent->GetPointID();
+				FRenderablePoint PointDescription = PointComponent->GeneratePointDescription();
 
-			PointSet->SetPointPosition(PointID, PointDescription.Position);			
-			PointSet->SetPointColor(PointID, PointDescription.Color);
-			PointSet->SetPointSize(PointID, PointDescription.Size);
+				ResolvedPointSet->SetPointPosition(PointID, PointDescription.Position);			
+				ResolvedPointSet->SetPointColor(PointID, PointDescription.Color);
+				ResolvedPointSet->SetPointSize(PointID, PointDescription.Size);
+			}
 		}
 	}
 }
 
 UScriptableToolPoint* UScriptableToolPointSet::AddPoint()
 {
-	PointComponents.Add(NewObject<UScriptableToolPoint>(this));
-	UScriptableToolPoint* NewPoint = PointComponents.Last();
-	NewPoint->SetPointID(PointSet->AddPoint(NewPoint->GeneratePointDescription()));
+	UScriptableToolPoint* NewPoint = nullptr;
+	if (UPointSetComponent* ResolvedPointSet = WeakPointSet.Get())
+	{
+		PointComponents.Add(NewObject<UScriptableToolPoint>(this));
+		NewPoint = PointComponents.Last();
+		NewPoint->SetPointID(ResolvedPointSet->AddPoint(NewPoint->GeneratePointDescription()));
+	}
 	return NewPoint;
 }
 
 void UScriptableToolPointSet::RemovePoint(UScriptableToolPoint* Point)
 {
-	if (ensure(Point))
+	if (UPointSetComponent* ResolvedPointSet = WeakPointSet.Get())
 	{
-		PointSet->RemovePoint(Point->GetPointID());
-		PointComponents.Remove(Point);
+		if (ensure(Point))
+		{
+			ResolvedPointSet->RemovePoint(Point->GetPointID());
+			PointComponents.Remove(Point);
+		}
 	}
 }
 
 void UScriptableToolPointSet::RemoveAllPoints()
 {
-	PointSet->Clear();
+	if (UPointSetComponent* ResolvedPointSet = WeakPointSet.Get())
+	{
+		ResolvedPointSet->Clear();
+	}
 	PointComponents.Empty();
 }
 
 void UScriptableToolPointSet::SetAllPointsColor(FColor Color)
 {
-	PointSet->SetAllPointsColor(Color);
+	if (UPointSetComponent* ResolvedPointSet = WeakPointSet.Get())
+	{
+		ResolvedPointSet->SetAllPointsColor(Color);
+	}
 }
 
 void UScriptableToolPointSet::SetAllPointsSize(float Size)
 {
-	PointSet->SetAllPointsSize(Size);
+	if (UPointSetComponent* ResolvedPointSet = WeakPointSet.Get())
+	{
+		ResolvedPointSet->SetAllPointsSize(Size);
+	}
 }
\ No newline at end of file
diff --git a/Engine/Plugins/Runtime/ScriptableToolsFramework/Source/ScriptableToolsFramework/Private/Drawing/ScriptableToolTriangleSet.cpp b/Engine/Plugins/Runtime/ScriptableToolsFramework/Source/ScriptableToolsFramework/Private/Drawing/ScriptableToolTriangleSet.cpp
--- a/Engine/Plugins/Runtime/ScriptableToolsFramework/Source/ScriptableToolsFramework/Private/Drawing/ScriptableToolTriangleSet.cpp
+++ b/Engine/Plugins/Runtime/ScriptableToolsFramework/Source/ScriptableToolsFramework/Private/Drawing/ScriptableToolTriangleSet.cpp
@@ -9,97 +9,122 @@
 void UScriptableToolTriangleSet::Initialize(TObjectPtr<UPreviewGeometry> PreviewGeometry)
 {
 	FString TriangleSetID = FGuid::NewGuid().ToString();
-	TriangleSet = PreviewGeometry->AddTriangleSet(TriangleSetID);
+	WeakTriangleSet = PreviewGeometry->AddTriangleSet(TriangleSetID);
 }
 
 void UScriptableToolTriangleSet::OnTick()
 {
-	for (TObjectPtr<UScriptableToolTriangle> TriangleComponent : TriangleComponents)
+	if (UTriangleSetComponent* ResolvedTriangleSet = WeakTriangleSet.Get())
 	{
-		if (TriangleComponent->IsDirty())
+		for (TObjectPtr<UScriptableToolTriangle> TriangleComponent : TriangleComponents)
 		{
-			int32 TriangleID = TriangleComponent->GetTriangleID();
-			FRenderableTriangle TriangleDescription = TriangleComponent->GenerateTriangleDescription();
-
-			TriangleSet->RemoveTriangle(TriangleID);
-			TriangleSet->InsertTriangle(TriangleID, TriangleDescription);
+			if (TriangleComponent->IsDirty())
+			{
+				int32 TriangleID = TriangleComponent->GetTriangleID();
+				FRenderableTriangle TriangleDescription = TriangleComponent->GenerateTriangleDescription();
+
+				ResolvedTriangleSet->RemoveTriangle(TriangleID);
+				ResolvedTriangleSet->InsertTriangle(TriangleID, TriangleDescription);
+			}
 		}
-	}
 
-	for (TObjectPtr<UScriptableToolQuad> QuadComponent : QuadComponents)
-	{
-		if (QuadComponent->IsDirty())
+		for (TObjectPtr<UScriptableToolQuad> QuadComponent : QuadComponents)
 		{
-			int32 TriangleAID = QuadComponent->GetTriangleAID();
-			int32 TriangleBID = QuadComponent->GetTriangleBID();
+			if (QuadComponent->IsDirty())
+			{
+				int32 TriangleAID = QuadComponent->GetTriangleAID();
+				int32 TriangleBID = QuadComponent->GetTriangleBID();
 
-			TPair<FRenderableTriangle, FRenderableTriangle> TriangleDescriptions = QuadComponent->GenerateQuadDescription();
+				TPair<FRenderableTriangle, FRenderableTriangle> TriangleDescriptions = QuadComponent->GenerateQuadDescription();
 
-			TriangleSet->RemoveTriangle(TriangleAID);
-			TriangleSet->InsertTriangle(TriangleAID, TriangleDescriptions.Key);
-			TriangleSet->RemoveTriangle(TriangleBID);
-			TriangleSet->InsertTriangle(TriangleBID, TriangleDescriptions.Value);
+				ResolvedTriangleSet->RemoveTriangle(TriangleAID);
+				ResolvedTriangleSet->InsertTriangle(TriangleAID, TriangleDescriptions.Key);
+				ResolvedTriangleSet->RemoveTriangle(TriangleBID);
+				ResolvedTriangleSet->InsertTriangle(TriangleBID, TriangleDescriptions.Value);
 
+			}
 		}
 	}
 }
 
 UScriptableToolTriangle* UScriptableToolTriangleSet::AddTriangle()
 {
-	TriangleComponents.Add(NewObject<UScriptableToolTriangle>(this));
-	UScriptableToolTriangle* NewTriangle = TriangleComponents.Last();
-	NewTriangle->SetTriangleID(TriangleSet->AddTriangle(NewTriangle->GenerateTriangleDescription()));
+	UScriptableToolTriangle* NewTriangle = nullptr;
+	if (UTriangleSetComponent* ResolvedTriangleSet = WeakTriangleSet.Get())
+	{
+		TriangleComponents.Add(NewObject<UScriptableToolTriangle>(this));
+		NewTriangle = TriangleComponents.Last();
+		NewTriangle->SetTriangleID(ResolvedTriangleSet->AddTriangle(NewTriangle->GenerateTriangleDescription()));
+	}
 	return NewTriangle;
 }
 
 UScriptableToolQuad* UScriptableToolTriangleSet::AddQuad()
 {
-	QuadComponents.Add(NewObject<UScriptableToolQuad>(this));
-	UScriptableToolQuad* NewQuad = QuadComponents.Last();
-
-	TPair<FRenderableTriangle, FRenderableTriangle> TriDescriptions = NewQuad->GenerateQuadDescription();
+	UScriptableToolQuad* NewQuad = nullptr;
+	if (UTriangleSetComponent* ResolvedTriangleSet = WeakTriangleSet.Get())
+	{
+		QuadComponents.Add(NewObject<UScriptableToolQuad>(this));
+		NewQuad = QuadComponents.Last();
 
-	int32 TriAID = TriangleSet->AddTriangle(TriDescriptions.Key);
-	int32 TriBID = TriangleSet->AddTriangle(TriDescriptions.Value);
+		TPair<FRenderableTriangle, FRenderableTriangle> TriDescriptions = NewQuad->GenerateQuadDescription();
 
-	NewQuad->SetTriangleAID(TriAID);
-	NewQuad->SetTriangleBID(TriBID);
+		int32 TriAID = ResolvedTriangleSet->AddTriangle(TriDescriptions.Key);
+		int32 TriBID = ResolvedTriangleSet->AddTriangle(TriDescriptions.Value);
 
+		NewQuad->SetTriangleAID(TriAID);
+		NewQuad->SetTriangleBID(TriBID);
+	}
 	return NewQuad;
 }
 
 void UScriptableToolTriangleSet::RemoveTriangle(UScriptableToolTriangle* Triangle)
 {
-	if (ensure(Triangle))
+	if (UTriangleSetComponent* ResolvedTriangleSet = WeakTriangleSet.Get())
 	{
-		TriangleSet->RemoveTriangle(Triangle->GetTriangleID());
-		TriangleComponents.Remove(Triangle);
+		if (ensure(Triangle))
+		{
+			ResolvedTriangleSet->RemoveTriangle(Triangle->GetTriangleID());
+			TriangleComponents.Remove(Triangle);
+		}
 	}
 }
 
 void UScriptableToolTriangleSet::RemoveQuad(UScriptableToolQuad* Quad)
 {
-	if (ensure(Quad))
+	if (UTriangleSetComponent* ResolvedTriangleSet = WeakTriangleSet.Get())
 	{
-		TriangleSet->RemoveTriangle(Quad->GetTriangleAID());
-		TriangleSet->RemoveTriangle(Quad->GetTriangleBID());
-		QuadComponents.Remove(Quad);
+		if (ensure(Quad))
+		{
+			ResolvedTriangleSet->RemoveTriangle(Quad->GetTriangleAID());
+			ResolvedTriangleSet->RemoveTriangle(Quad->GetTriangleBID());
+			QuadComponents.Remove(Quad);
+		}
 	}
 }
 
 
 void UScriptableToolTriangleSet::RemoveAllFaces()
 {
-	TriangleSet->Clear();
+	if (UTriangleSetComponent* ResolvedTriangleSet = WeakTriangleSet.Get())
+	{
+		ResolvedTriangleSet->Clear();
+	}
 	TriangleComponents.Empty();
 }
 
 void UScriptableToolTriangleSet::SetAllTrianglesColor(FColor Color)
 {
-	TriangleSet->SetAllTrianglesColor(Color);
+	if (UTriangleSetComponent* ResolvedTriangleSet = WeakTriangleSet.Get())
+	{
+		ResolvedTriangleSet->SetAllTrianglesColor(Color);
+	}
 }
 
 void UScriptableToolTriangleSet::SetAllTrianglesMaterial(UMaterialInterface* Material)
 {
-	TriangleSet->SetAllTrianglesMaterial(Material);
+	if (UTriangleSetComponent* ResolvedTriangleSet = WeakTriangleSet.Get())
+	{
+		ResolvedTriangleSet->SetAllTrianglesMaterial(Material);
+	}
 }
\ No newline at end of file
diff --git a/Engine/Plugins/Runtime/ScriptableToolsFramework/Source/ScriptableToolsFramework/Public/Drawing/ScriptableToolLineSet.h b/Engine/Plugins/Runtime/ScriptableToolsFramework/Source/ScriptableToolsFramework/Public/Drawing/ScriptableToolLineSet.h
--- a/Engine/Plugins/Runtime/ScriptableToolsFramework/Source/ScriptableToolsFramework/Public/Drawing/ScriptableToolLineSet.h
+++ b/Engine/Plugins/Runtime/ScriptableToolsFramework/Source/ScriptableToolsFramework/Public/Drawing/ScriptableToolLineSet.h
@@ -63,9 +63,13 @@ protected:
 	UPROPERTY(Transient, DuplicateTransient, NonTransactional, SkipSerialization)
 	TObjectPtr<UPreviewGeometry> ToolDrawableGeometry = nullptr;
 
+	UE_DEPRECATED(5.7, "Use WeakLineSet instead.")
 	UPROPERTY(Transient, DuplicateTransient, NonTransactional, SkipSerialization)
 	TObjectPtr<ULineSetComponent> LineSet = nullptr;
 
+	UPROPERTY(Transient, DuplicateTransient, NonTransactional, SkipSerialization)
+	TWeakObjectPtr<ULineSetComponent> WeakLineSet = nullptr;
+
 	UPROPERTY(Transient, DuplicateTransient, NonTransactional, SkipSerialization)
 	TArray<TObjectPtr<UScriptableToolLine>> LineComponents;
 
diff --git a/Engine/Plugins/Runtime/ScriptableToolsFramework/Source/ScriptableToolsFramework/Public/Drawing/ScriptableToolPointSet.h b/Engine/Plugins/Runtime/ScriptableToolsFramework/Source/ScriptableToolsFramework/Public/Drawing/ScriptableToolPointSet.h
--- a/Engine/Plugins/Runtime/ScriptableToolsFramework/Source/ScriptableToolsFramework/Public/Drawing/ScriptableToolPointSet.h
+++ b/Engine/Plugins/Runtime/ScriptableToolsFramework/Source/ScriptableToolsFramework/Public/Drawing/ScriptableToolPointSet.h
@@ -62,9 +62,13 @@ protected:
 	UPROPERTY(Transient, DuplicateTransient, NonTransactional, SkipSerialization)
 	TObjectPtr<UPreviewGeometry> ToolDrawableGeometry = nullptr;
 
+	UE_DEPRECATED(5.7, "Use WeakPointSet instead.")
 	UPROPERTY(Transient, DuplicateTransient, NonTransactional, SkipSerialization)
 	TObjectPtr<UPointSetComponent> PointSet = nullptr;
 
+	UPROPERTY(Transient, DuplicateTransient, NonTransactional, SkipSerialization)
+	TWeakObjectPtr<UPointSetComponent> WeakPointSet = nullptr;
+
 	UPROPERTY(Transient, DuplicateTransient, NonTransactional, SkipSerialization)
 	TArray<TObjectPtr<UScriptableToolPoint>> PointComponents;
 
diff --git a/Engine/Plugins/Runtime/ScriptableToolsFramework/Source/ScriptableToolsFramework/Public/Drawing/ScriptableToolTriangleSet.h b/Engine/Plugins/Runtime/ScriptableToolsFramework/Source/ScriptableToolsFramework/Public/Drawing/ScriptableToolTriangleSet.h
--- a/Engine/Plugins/Runtime/ScriptableToolsFramework/Source/ScriptableToolsFramework/Public/Drawing/ScriptableToolTriangleSet.h
+++ b/Engine/Plugins/Runtime/ScriptableToolsFramework/Source/ScriptableToolsFramework/Public/Drawing/ScriptableToolTriangleSet.h
@@ -76,14 +76,18 @@ protected:
 	UPROPERTY(Transient, DuplicateTransient, NonTransactional, SkipSerialization)
 	TObjectPtr<UPreviewGeometry> ToolDrawableGeometry = nullptr;
 
+	UE_DEPRECATED(5.7, "Use WeakTriangleSet instead.")
 	UPROPERTY(Transient, DuplicateTransient, NonTransactional, SkipSerialization)
 	TObjectPtr<UTriangleSetComponent> TriangleSet = nullptr;
 
+	UPROPERTY(Transient, DuplicateTransient, NonTransactional, SkipSerialization)
+	TWeakObjectPtr<UTriangleSetComponent> WeakTriangleSet = nullptr;
+
 	UPROPERTY(Transient, DuplicateTransient, NonTransactional, SkipSerialization)
 	TArray<TObjectPtr<UScriptableToolTriangle>> TriangleComponents;
 
 	UPROPERTY(Transient, DuplicateTransient, NonTransactional, SkipSerialization)
-	TArray < TObjectPtr < UScriptableToolQuad>> QuadComponents;
+	TArray<TObjectPtr<UScriptableToolQuad>> QuadComponents;
 
 };
 
