ScriptableTools: Fixed issue where setting values on a ToolHeader widget during OnScriptSetup would fail.

diff --git a/Engine/Plugins/Runtime/ScriptableToolsFramework/Source/ScriptableToolsFramework/Private/ScriptableInteractiveTool.cpp b/Engine/Plugins/Runtime/ScriptableToolsFramework/Source/ScriptableToolsFramework/Private/ScriptableInteractiveTool.cpp
--- a/Engine/Plugins/Runtime/ScriptableToolsFramework/Source/ScriptableToolsFramework/Private/ScriptableInteractiveTool.cpp
+++ b/Engine/Plugins/Runtime/ScriptableToolsFramework/Source/ScriptableToolsFramework/Private/ScriptableInteractiveTool.cpp
@@ -1327,22 +1327,19 @@ UScriptableToolTriangleSet* UScriptableInteractiveTool::AddTriangleSet()
 void UScriptableInteractiveTool::SetToolPanelToolHeaderWidget(UUserWidget* Widget)
 {
 	ToolPanelToolHeaderWidget = Widget;
+
+	// Must invoke TakeWidget here to ensure that the underlying Slate widget is created
+	// on this call, otherwise calls to set properties on the widget will fail during
+	// OnScriptSetup.
+	if (ToolPanelToolHeaderWidget)
+	{
+		ToolPanelToolHeaderWidget->TakeWidget();
+	}
 }
 
 TSharedRef<SWidget> UScriptableInteractiveTool::GetToolPanelToolHeaderWidget() const
 {
-	if (ToolPanelToolHeaderWidget)
-	{
-		TSharedRef<SWidget> WrappedWidget = ToolPanelToolHeaderWidget->TakeDerivedWidget<SObjectWidget>([](UUserWidget* Widget, TSharedRef<SWidget> Content)
-		{
-			return SNew(SObjectWidget, Widget)
-				[
-					Content
-				];
-		});
-		return WrappedWidget;
-	}
-	return SNullWidget::NullWidget;
+	return ToolPanelToolHeaderWidget ? ToolPanelToolHeaderWidget->TakeWidget() : SNullWidget::NullWidget;
 }
 
 
