Fixed category assignment for SetFocusInViewport

#rnx

diff --git a/Engine/Plugins/Runtime/ScriptableToolsFramework/Source/ScriptableToolsFramework/Public/ScriptableInteractiveTool.h b/Engine/Plugins/Runtime/ScriptableToolsFramework/Source/ScriptableToolsFramework/Public/ScriptableInteractiveTool.h
--- a/Engine/Plugins/Runtime/ScriptableToolsFramework/Source/ScriptableToolsFramework/Public/ScriptableInteractiveTool.h
+++ b/Engine/Plugins/Runtime/ScriptableToolsFramework/Source/ScriptableToolsFramework/Public/ScriptableInteractiveTool.h
@@ -956,7 +956,7 @@ public:
 
 	// Viewport Focus API
 
-	UFUNCTION(BlueprintCallable, Category = "ScriptableTools|Viewport")
+	UFUNCTION(BlueprintCallable, Category = "ScriptableTool|Viewport")
 	UE_API void SetFocusInViewport();
 };
 
