ScriptableTools: Fixed error in AddSingleClickBehavior and AddDoubleClickBehavior function tooltip.

diff --git a/Engine/Plugins/Runtime/ScriptableToolsFramework/Source/ScriptableToolsFramework/Public/BaseTools/ScriptableModularBehaviorTool.h b/Engine/Plugins/Runtime/ScriptableToolsFramework/Source/ScriptableToolsFramework/Public/BaseTools/ScriptableModularBehaviorTool.h
--- a/Engine/Plugins/Runtime/ScriptableToolsFramework/Source/ScriptableToolsFramework/Public/BaseTools/ScriptableModularBehaviorTool.h
+++ b/Engine/Plugins/Runtime/ScriptableToolsFramework/Source/ScriptableToolsFramework/Public/BaseTools/ScriptableModularBehaviorTool.h
@@ -36,7 +36,7 @@ public:
 	*   Implements a standard "button-click"-style input behavior
 	*   The state machine works as follows:
 	*    1) on input-device-button-press, hit-test the target. If hit, begin capture
-	*    2) on input-device-button-release, hit-test the target. If hit, call OnClicked(). If not hit, ignore click.
+	*    2) on input-device-button-release, hit-test the target. If hit, call OnHitByClicked(). If not hit, ignore click.
 	*    
 	*   The second hit-test is required to allow the click to be "cancelled" by moving away
 	*   from the target. This is standard GUI behavior. You can disable this second hit test
@@ -62,7 +62,7 @@ public:
 	*   Implements a standard "button-click"-style input behavior
 	*   The state machine works as follows:
 	*    1) on input-device-button-press, hit-test the target. If hit, begin capture
-	*    2) on input-device-button-release, hit-test the target. If hit, call OnClicked(). If not hit, ignore click.
+	*    2) on input-device-button-release, hit-test the target. If hit, call OnHitByClicked(). If not hit, ignore click.
 	*    
 	*   The second hit-test is required to allow the click to be "cancelled" by moving away
 	*   from the target. This is standard GUI behavior. You can disable this second hit test
