ScriptableTools: Fix the usage of LogTemp for a new Scriptable Tools log category and change the non-highlighted log message to write to the console in addition to the log file.

diff --git a/Engine/Plugins/Runtime/ScriptableToolsFramework/Source/ScriptableToolsFramework/Private/ScriptableInteractiveTool.cpp b/Engine/Plugins/Runtime/ScriptableToolsFramework/Source/ScriptableToolsFramework/Private/ScriptableInteractiveTool.cpp
--- a/Engine/Plugins/Runtime/ScriptableToolsFramework/Source/ScriptableToolsFramework/Private/ScriptableInteractiveTool.cpp
+++ b/Engine/Plugins/Runtime/ScriptableToolsFramework/Source/ScriptableToolsFramework/Private/ScriptableInteractiveTool.cpp
@@ -35,6 +35,8 @@
 
 #define LOCTEXT_NAMESPACE "UScriptableInteractiveTool"
 
+DECLARE_LOG_CATEGORY_CLASS(LogScriptableInteractiveTool, Log, All);
+
 namespace ScriptableInteractiveToolLocals
 {
 	class SScriptableToolViewportOverlayDragContainer : public SObjectWidget
@@ -1225,11 +1227,11 @@ void UScriptableInteractiveTool::AddLogMessage(FText Message, bool bHighlighted)
 	{
 		if (bHighlighted)
 		{
-			UE_LOG(LogTemp, Warning, TEXT("%s"), *Message.ToString() );
+			UE_LOG(LogScriptableInteractiveTool, Warning, TEXT("%s"), *Message.ToString() );
 		}
 		else
 		{
-			UE_LOG(LogTemp, Log, TEXT("%s"), *Message.ToString() );
+			UE_LOG(LogScriptableInteractiveTool, Display, TEXT("%s"), *Message.ToString() );
 		}
 	}
 }
