ScriptableTools: Added MakeTransactional BP function to ScriptableInteractiveToolPropertySet so that tools can optionally enable transactions on a property set.

diff --git a/Engine/Plugins/Runtime/ScriptableToolsFramework/Source/ScriptableToolsFramework/Private/ScriptableInteractiveTool.cpp b/Engine/Plugins/Runtime/ScriptableToolsFramework/Source/ScriptableToolsFramework/Private/ScriptableInteractiveTool.cpp
--- a/Engine/Plugins/Runtime/ScriptableToolsFramework/Source/ScriptableToolsFramework/Private/ScriptableInteractiveTool.cpp
+++ b/Engine/Plugins/Runtime/ScriptableToolsFramework/Source/ScriptableToolsFramework/Private/ScriptableInteractiveTool.cpp
@@ -132,6 +132,12 @@ UScriptableInteractiveToolPropertySet* UScriptableInteractiveToolPropertySet::Se
 	return this;
 }
 
+UScriptableInteractiveToolPropertySet* UScriptableInteractiveToolPropertySet::MakeTransactional()
+{
+	SetFlags(RF_Transactional);
+	return this;
+}
+
 UScriptableInteractiveTool* UScriptableInteractiveToolPropertySet::GetOwningTool(EToolsFrameworkOutcomePins& Outcome)
 {
 	if (UScriptableInteractiveTool* OwningTool = ParentTool.Get())
diff --git a/Engine/Plugins/Runtime/ScriptableToolsFramework/Source/ScriptableToolsFramework/Public/ScriptableInteractiveTool.h b/Engine/Plugins/Runtime/ScriptableToolsFramework/Source/ScriptableToolsFramework/Public/ScriptableInteractiveTool.h
--- a/Engine/Plugins/Runtime/ScriptableToolsFramework/Source/ScriptableToolsFramework/Public/ScriptableInteractiveTool.h
+++ b/Engine/Plugins/Runtime/ScriptableToolsFramework/Source/ScriptableToolsFramework/Public/ScriptableInteractiveTool.h
@@ -289,6 +289,9 @@ class UScriptableInteractiveToolPropertySet : public UInteractiveToolPropertySet
 	UFUNCTION(BlueprintCallable, Category = "ScriptableTool")
 	UE_API UPARAM(DisplayName = "Property Set") UScriptableInteractiveToolPropertySet* SetPropertyAsHidden(FString PropertyName, bool bHidden);
 
+	UFUNCTION(BlueprintCallable, Category = "ScriptableTool")
+	UE_API UPARAM(DisplayName = "Property Set") UScriptableInteractiveToolPropertySet* MakeTransactional(); 
+
 	UE_API virtual UWorld* GetWorld() const override;	// must implement this for WorldContextObject dependent BP functions to show.
 
 protected:
