Updating tooltips to make Capture Priority values clearer.
Fix typo preventing "Key" property tooltip from appearing

Add description and example to Mouse Hover and Mouse Wheel Behaviors

#rnx
#Synced-CL 50529789

diff --git a/Engine/Plugins/Runtime/ScriptableToolsFramework/Source/ScriptableToolsFramework/Public/BaseTools/ScriptableModularBehaviorTool.h b/Engine/Plugins/Runtime/ScriptableToolsFramework/Source/ScriptableToolsFramework/Public/BaseTools/ScriptableModularBehaviorTool.h
--- a/Engine/Plugins/Runtime/ScriptableToolsFramework/Source/ScriptableToolsFramework/Public/BaseTools/ScriptableModularBehaviorTool.h
+++ b/Engine/Plugins/Runtime/ScriptableToolsFramework/Source/ScriptableToolsFramework/Public/BaseTools/ScriptableModularBehaviorTool.h
@@ -45,7 +45,7 @@ public:
 	*	@param IfHitByClick Test if hit by a click
 	*	@param OnHitByClick Notify that click occurred
 	*	@param CaptureCheck Only enable capture if returns true
-	*	@param CapturePriority The priority is used to resolve situations where multiple behaviors want the same capture
+	*	@param CapturePriority The priority is used to resolve situations where multiple behaviors want the same capture. Lower values, including negative, have higher priority.
 	*	@param MouseButton Determines which mouse button the behavior captures on
 	*/
 	UFUNCTION(BlueprintCallable, Category = "ScriptableTool|Input", meta = (AdvancedDisplay = "CaptureCheck, CapturePriority, MouseButton, bHitTestOnRelease"))
@@ -71,7 +71,7 @@ public:
 	*	@param IfHitByClick Test if hit by a click
 	*	@param OnHitByClick Notify that double click occurred
 	*	@param CaptureCheck Only enable capture if returns true
-	*	@param CapturePriority The priority is used to resolve situations where multiple behaviors want the same capture
+	*	@param CapturePriority The priority is used to resolve situations where multiple behaviors want the same capture. Lower values, including negative, have higher priority.
 	*	@param MouseButton Determines which mouse button the behavior captures on
 	*/
 	UFUNCTION(BlueprintCallable, Category = "ScriptableTool|Input", meta = (AdvancedDisplay = "CaptureCheck, CapturePriority, MouseButton, bHitTestOnRelease"))
@@ -100,7 +100,7 @@ public:
 	*	@param OnClickRelease Notify Target that click release occurred
 	*	@param OnTerminateDragSequence Notify Target that click-drag sequence has been explicitly terminated (eg by escape key)
 	*	@param CaptureCheck Only enable capture if returns true
-	*	@param CapturePriority The priority is used to resolve situations where multiple behaviors want the same capture
+	*	@param CapturePriority The priority is used to resolve situations where multiple behaviors want the same capture. Lower values, including negative, have higher priority.
 	*	@param MouseButton Determines which mouse button the behavior captures on
 	*/
 	UFUNCTION(BlueprintCallable, Category = "ScriptableTool|Input", meta = (AdvancedDisplay = "CaptureCheck, CapturePriority, MouseButton, bUpdateModifiersDuringDrag"))
@@ -158,11 +158,15 @@ public:
 	);
 
 	/**
+	*	
+	*   AddMouseWheelBehavior will run the event "OnMouseWheelScrollUp" and "OnMouseWheelScrollDown" each time the mouse wheel is rotated.
+	*	It is up to the user to record these in variables or otherwise react to these events.
+	* 
 	*	@param TestShouldRespondToMouseWheel The result's bHit property determines whether the mouse wheel action will be captured. (Perhaps the mouse wheel only does something when mousing over some part of a mesh)
 	*	@param OnMouseWheelScrollUp CurrentPos device position/ray at point where mouse wheel is engaged
 	*	@param OnMouseWheelScrollDown CurrentPos device position/ray at point where mouse wheel is engaged
 	*	@param CaptureCheck Only enable capture if returns true
-	*	@param CapturePriority The priority is used to resolve situations where multiple behaviors want the same capture
+	*	@param CapturePriority The priority is used to resolve situations where multiple behaviors want the same capture. Lower values, including negative, have higher priority.
 	*/
 	UFUNCTION(BlueprintCallable, Category = "ScriptableTool|Input", meta = (AdvancedDisplay = "CaptureCheck, CapturePriority"))
 	UE_API void AddMouseWheelBehavior(		
@@ -215,12 +219,16 @@ public:
 	);
 
 	/**
-	*	@param BeginHoverSequenceHitTest  Do hover hit-test
+	* 
+	*   AddMouseHoverBehavior implements a generic mouse hover behavior.
+	*   For example, this can be used to change the "mode" of the current tool depending on what actor (if any) the mouse is hovering over.
+	* 
+	*	@param BeginHoverSequenceHitTest  Do hover hit-test. If this returns False, no further events below are executed.
 	*	@param OnBeginHover Initialize hover sequence at given position
 	*	@param OnUpdateHover Update active hover sequence with new input position
 	*	@param OnEndHover Terminate active hover sequence
 	*   @param HoverCaptureCheck  Only enable hover capture if returns true
-	*	@param CapturePriority The priority is used to resolve situations where multiple behaviors want the same capture
+	*	@param CapturePriority The priority is used to resolve situations where multiple behaviors want the same capture. Lower values, including negative, have higher priority.
 	*/
 	UFUNCTION(BlueprintCallable, Category = "ScriptableTool|Input", meta = (AdvancedDisplay = "CaptureCheck, CapturePriority"))
 	UE_API void AddMouseHoverBehavior(
@@ -233,15 +241,19 @@ public:
 	);
 
 	/**
-	*	AddSingleKeyInputBehavior implements a generic keyboard key listener behavior
+	*	AddSingleKeyInputBehavior implements a generic keyboard key listener behavior.
+	*   This is often used to add "hotkeys" or "shortcuts" to the tool.
+	*   
+	*   CapturePriority can be used to either give priority to "built in" viewport hotkeys (the default)
+	*   or to have this behavior's key be the priority, using CapturePriority numbers lower than 50, for example.
 	*	
 	*	@param OnKeyPressed Callback when the target key is pressed.
 	*	@param OnKeyReleased Callback when the target key is released
 	*	@param OnForceEndCaptureFuncIn Callback when capture is ended prematurely, typically due to the viewport losing
 	*	 focus, in which case the release callback will not be called.
-	*	@param Key,Target key to watch for
+	*	@param Key Target key to watch for
 	*	@param CaptureCheck Only enable capture if returns true
-	*	@param CapturePriority The priority is used to resolve situations where multiple behaviors want the same capture
+	*	@param CapturePriority The priority is used to resolve situations where multiple behaviors want the same capture. Lower values, including negative, have higher priority.
 	*/
 
 
@@ -274,9 +286,9 @@ public:
 	*	@param OnForceEndCaptureFuncIn Callback when capture is ended prematurely, typically due to the viewport losing
 	*	 focus, in which case the release callback will not be called.
 	*	@param Keys Target keys to watch for
-	*	@param bRequireAllKeys If true, all target keys must be pressed simultaniously to recieve press/release events. Otherwise, any and all keys can trigger events.
+	*	@param bRequireAllKeys If true, all target keys must be pressed simultaniously to receive press/release events. Otherwise, any and all keys can trigger events.
 	*	@param CaptureCheck Only enable capture if returns true
-	*	@param CapturePriority The priority is used to resolve situations where multiple behaviors want the same capture
+	*	@param CapturePriority The priority is used to resolve situations where multiple behaviors want the same capture. Lower values, including negative, have higher priority.
 	*/
 
 
