ScriptableToolsFramework: Only reference ToolWidgets if building developer tools

#rnx

#Synced-CL 50554652

diff --git a/Engine/Plugins/Runtime/ScriptableToolsFramework/Source/ScriptableToolsFramework/Private/ScriptableInteractiveTool.cpp b/Engine/Plugins/Runtime/ScriptableToolsFramework/Source/ScriptableToolsFramework/Private/ScriptableInteractiveTool.cpp
--- a/Engine/Plugins/Runtime/ScriptableToolsFramework/Source/ScriptableToolsFramework/Private/ScriptableInteractiveTool.cpp
+++ b/Engine/Plugins/Runtime/ScriptableToolsFramework/Source/ScriptableToolsFramework/Private/ScriptableInteractiveTool.cpp
@@ -23,7 +23,6 @@
 
 #include "Blueprint/UserWidget.h"
 #include "ContextObjectStore.h"
-#include "Overlay/SDraggableBoxOverlay.h"
 #include "Utility/ScriptableToolContextObjects.h"
 #include "Slate/SObjectWidget.h"
 #include "Components/Widget.h"
@@ -31,6 +30,10 @@
 
 #include "UObject/EnumProperty.h"
 
+#if WITH_UNREAL_DEVELOPER_TOOLS
+#include "Overlay/SDraggableBoxOverlay.h"
+#endif
+
 #include UE_INLINE_GENERATED_CPP_BY_NAME(ScriptableInteractiveTool)
 
 #define LOCTEXT_NAMESPACE "UScriptableInteractiveTool"
@@ -1362,7 +1365,7 @@ void UScriptableInteractiveTool::SetOverlayWidget(UUserWidget* Widget, bool bMak
 	{
 		if (bMakeDraggable)
 		{
-
+#if WITH_UNREAL_DEVELOPER_TOOLS
 			TSharedRef<SWidget> WrappedWidget = Widget->TakeDerivedWidget<ScriptableInteractiveToolLocals::SScriptableToolViewportOverlayDragContainer>([](UUserWidget* Widget, TSharedRef<SWidget> Content) {
 				return SNew(ScriptableInteractiveToolLocals::SScriptableToolViewportOverlayDragContainer, Widget)
 					[
@@ -1372,9 +1375,10 @@ void UScriptableInteractiveTool::SetOverlayWidget(UUserWidget* Widget, bool bMak
 						]
 					];
 				});
-
+#else
+			TSharedRef<SWidget> WrappedWidget = Widget->TakeWidget();
+#endif
 			ViewportWidgetAPI->SetOverlayWidget(WrappedWidget);
-
 		}
 		else
 		{
diff --git a/Engine/Plugins/Runtime/ScriptableToolsFramework/Source/ScriptableToolsFramework/ScriptableToolsFramework.Build.cs b/Engine/Plugins/Runtime/ScriptableToolsFramework/Source/ScriptableToolsFramework/ScriptableToolsFramework.Build.cs
--- a/Engine/Plugins/Runtime/ScriptableToolsFramework/Source/ScriptableToolsFramework/ScriptableToolsFramework.Build.cs
+++ b/Engine/Plugins/Runtime/ScriptableToolsFramework/Source/ScriptableToolsFramework/ScriptableToolsFramework.Build.cs
@@ -45,10 +45,18 @@ public class ScriptableToolsFramework : ModuleRules
 				"Engine",
 				"Slate",
 				"SlateCore",
-				"ToolWidgets"
 			}
 			);
-
+		
+		if (Target.bBuildDeveloperTools)
+		{
+			PrivateDependencyModuleNames.AddRange(
+				new string[]
+				{
+					"ToolWidgets"
+				}
+				);
+		}
 
 		DynamicallyLoadedModuleNames.AddRange(
 			new string[]
