ScriptableTools: Refactor SDraggableBoxOverlay usage to isolate ToolWidgets dependency to ScriptableToolsEditorMode rather than ScriptableToolsFramework. The coupling of ScriptableToolsFramework on ToolWidgets prevented runtimes from shipping the framework due to ToolWidgets non-redistributable status as part of the Source/Developer modules.

#Synced-CL 52701182

diff --git a/Engine/Plugins/Editor/ScriptableToolsEditorMode/Source/ScriptableToolsEditorMode/Private/ScriptableToolsEditorMode.cpp b/Engine/Plugins/Editor/ScriptableToolsEditorMode/Source/ScriptableToolsEditorMode/Private/ScriptableToolsEditorMode.cpp
--- a/Engine/Plugins/Editor/ScriptableToolsEditorMode/Source/ScriptableToolsEditorMode/Private/ScriptableToolsEditorMode.cpp
+++ b/Engine/Plugins/Editor/ScriptableToolsEditorMode/Source/ScriptableToolsEditorMode/Private/ScriptableToolsEditorMode.cpp
@@ -39,6 +39,10 @@
 #include "Kismet2/BlueprintEditorUtils.h"
 #include "Slate/SceneViewport.h"
 
+#if WITH_UNREAL_DEVELOPER_TOOLS
+#include "Overlay/SDraggableBoxOverlay.h"
+#endif
+
 #include UE_INLINE_GENERATED_CPP_BY_NAME(ScriptableToolsEditorMode)
 
 
@@ -420,6 +424,15 @@ void UScriptableToolsEditorMode::InitializeModeContexts()
 			}
 		}
 		);
+#if WITH_UNREAL_DEVELOPER_TOOLS
+	ViewportWidgetAPI->MakeDraggableWidgetFunc = [](TSharedRef<SWidget> Content)
+		{
+			return SNew(UE::ToolWidgets::SDraggableBoxOverlay)
+				[
+					Content
+				];
+		};
+#endif
 	AddContextObject(ViewportWidgetAPI);
 
 	UScriptableToolViewportFocusAPI* ViewportFocusAPI = NewObject<UScriptableToolViewportFocusAPI>();
diff --git a/Engine/Plugins/Editor/ScriptableToolsEditorMode/Source/ScriptableToolsEditorMode/ScriptableToolsEditorMode.Build.cs b/Engine/Plugins/Editor/ScriptableToolsEditorMode/Source/ScriptableToolsEditorMode/ScriptableToolsEditorMode.Build.cs
--- a/Engine/Plugins/Editor/ScriptableToolsEditorMode/Source/ScriptableToolsEditorMode/ScriptableToolsEditorMode.Build.cs
+++ b/Engine/Plugins/Editor/ScriptableToolsEditorMode/Source/ScriptableToolsEditorMode/ScriptableToolsEditorMode.Build.cs
@@ -68,7 +68,16 @@ public class ScriptableToolsEditorMode : ModuleRules
 
 			}
 			);
-
+		
+		if (Target.bBuildDeveloperTools)
+		{
+			PrivateDependencyModuleNames.AddRange(
+			new string[]
+			{
+				"ToolWidgets"
+			}
+			);
+		}
 
 		DynamicallyLoadedModuleNames.AddRange(
 			new string[]
diff --git a/Engine/Plugins/Runtime/ScriptableToolsFramework/Source/ScriptableToolsFramework/Private/ScriptableInteractiveTool.cpp b/Engine/Plugins/Runtime/ScriptableToolsFramework/Source/ScriptableToolsFramework/Private/ScriptableInteractiveTool.cpp
--- a/Engine/Plugins/Runtime/ScriptableToolsFramework/Source/ScriptableToolsFramework/Private/ScriptableInteractiveTool.cpp
+++ b/Engine/Plugins/Runtime/ScriptableToolsFramework/Source/ScriptableToolsFramework/Private/ScriptableInteractiveTool.cpp
@@ -30,10 +30,6 @@
 
 #include "UObject/EnumProperty.h"
 
-#if WITH_UNREAL_DEVELOPER_TOOLS
-#include "Overlay/SDraggableBoxOverlay.h"
-#endif
-
 #include UE_INLINE_GENERATED_CPP_BY_NAME(ScriptableInteractiveTool)
 
 #define LOCTEXT_NAMESPACE "UScriptableInteractiveTool"
@@ -1365,20 +1361,20 @@ void UScriptableInteractiveTool::SetOverlayWidget(UUserWidget* Widget, bool bMak
 	{
 		if (bMakeDraggable)
 		{
-#if WITH_UNREAL_DEVELOPER_TOOLS
-			TSharedRef<SWidget> WrappedWidget = Widget->TakeDerivedWidget<ScriptableInteractiveToolLocals::SScriptableToolViewportOverlayDragContainer>([](UUserWidget* Widget, TSharedRef<SWidget> Content) {
+			if (ViewportWidgetAPI->MakeDraggableWidgetFunc)
+			{
+				TSharedRef<SWidget> WrappedWidget = Widget->TakeDerivedWidget<ScriptableInteractiveToolLocals::SScriptableToolViewportOverlayDragContainer>([&ViewportWidgetAPI](UUserWidget* Widget, TSharedRef<SWidget> Content) {
 				return SNew(ScriptableInteractiveToolLocals::SScriptableToolViewportOverlayDragContainer, Widget)
 					[
-						SNew(UE::ToolWidgets::SDraggableBoxOverlay)
-						[
-							Content
-						]
+						ViewportWidgetAPI->MakeDraggableWidgetFunc(Content)
 					];
 				});
-#else
-			TSharedRef<SWidget> WrappedWidget = Widget->TakeWidget();
-#endif
-			ViewportWidgetAPI->SetOverlayWidget(WrappedWidget);
+				ViewportWidgetAPI->SetOverlayWidget(WrappedWidget);
+			}
+			else
+			{
+				ViewportWidgetAPI->SetOverlayWidget(Widget->TakeWidget());
+			}
 		}
 		else
 		{
diff --git a/Engine/Plugins/Runtime/ScriptableToolsFramework/Source/ScriptableToolsFramework/Public/Utility/ScriptableToolContextObjects.h b/Engine/Plugins/Runtime/ScriptableToolsFramework/Source/ScriptableToolsFramework/Public/Utility/ScriptableToolContextObjects.h
--- a/Engine/Plugins/Runtime/ScriptableToolsFramework/Source/ScriptableToolsFramework/Public/Utility/ScriptableToolContextObjects.h
+++ b/Engine/Plugins/Runtime/ScriptableToolsFramework/Source/ScriptableToolsFramework/Public/Utility/ScriptableToolContextObjects.h
@@ -82,6 +82,8 @@ public:
 		ClearOverlayWidget();
 	}
 
+	TUniqueFunction<TSharedRef<SWidget>(TSharedRef<SWidget> Content)> MakeDraggableWidgetFunc;
+
 private:
 
 	TUniqueFunction<void(TSharedRef<SWidget> InOverlaidWidget)> ReplaceOverlaidWidgetFunc;
diff --git a/Engine/Plugins/Runtime/ScriptableToolsFramework/Source/ScriptableToolsFramework/ScriptableToolsFramework.Build.cs b/Engine/Plugins/Runtime/ScriptableToolsFramework/Source/ScriptableToolsFramework/ScriptableToolsFramework.Build.cs
--- a/Engine/Plugins/Runtime/ScriptableToolsFramework/Source/ScriptableToolsFramework/ScriptableToolsFramework.Build.cs
+++ b/Engine/Plugins/Runtime/ScriptableToolsFramework/Source/ScriptableToolsFramework/ScriptableToolsFramework.Build.cs
@@ -47,16 +47,6 @@ public class ScriptableToolsFramework : ModuleRules
 				"SlateCore",
 			}
 			);
-		
-		if (Target.bBuildDeveloperTools)
-		{
-			PrivateDependencyModuleNames.AddRange(
-				new string[]
-				{
-					"ToolWidgets"
-				}
-				);
-		}
 
 		DynamicallyLoadedModuleNames.AddRange(
 			new string[]
